Added Fracture plugin with FractureEngine module for library of the algorithm implementations for the ChaosEditor/FractureEditor

diff --git a/Engine/Plugins/Experimental/Fracture/Fracture.uplugin b/Engine/Plugins/Experimental/Fracture/Fracture.uplugin
new file mode 100644
--- /dev/null
+++ b/Engine/Plugins/Experimental/Fracture/Fracture.uplugin
@@ -0,0 +1,26 @@
+{
+	"FileVersion": 1,
+	"Version": 1,
+	"VersionName": "0.1",
+	"FriendlyName": "Fracture",
+	"Description": "Adds Module for FractureEditor",
+	"Category": "Physics",
+	"CreatedBy": "Epic Games, Inc.",
+	"CreatedByURL": "http://epicgames.com",
+	"DocsURL": "",
+	"MarketplaceURL": "",
+	"SupportURL": "",
+	"CanContainContent": false,
+	"IsBetaVersion": true,
+	"Installed": false,
+	"EnabledByDefault": false,
+	"Modules": [
+		{
+			"Name": "FractureEngine",
+			"Type": "Runtime",
+			"LoadingPhase": "PostEngineInit"
+		}
+	],
+	"Plugins": [
+	]
+}
\ No newline at end of file
diff --git a/Engine/Plugins/Experimental/Fracture/Source/FractureEngine/FractureEngine.Build.cs b/Engine/Plugins/Experimental/Fracture/Source/FractureEngine/FractureEngine.Build.cs
new file mode 100644
--- /dev/null
+++ b/Engine/Plugins/Experimental/Fracture/Source/FractureEngine/FractureEngine.Build.cs
@@ -0,0 +1,28 @@
+// Copyright Epic Games, Inc. All Rights Reserved.
+
+using UnrealBuildTool;
+
+public class FractureEngine : ModuleRules
+{
+	public FractureEngine(ReadOnlyTargetRules Target) : base(Target)
+	{
+		PrivateIncludePaths.Add(ModuleDirectory + "/Private");
+		PublicIncludePaths.Add(ModuleDirectory + "/Public");
+		PublicDependencyModuleNames.AddRange(
+			new string[]
+			{
+
+				// ... add other public dependencies that you statically link with here ...
+			}
+            );
+			
+		
+		PrivateDependencyModuleNames.AddRange(
+			new string[]
+			{
+
+				// ... add private dependencies that you statically link with here ...	
+			}
+            );
+	}
+}
diff --git a/Engine/Plugins/Experimental/Fracture/Source/FractureEngine/Private/FractureToolClustering.cpp b/Engine/Plugins/Experimental/Fracture/Source/FractureEngine/Private/FractureToolClustering.cpp
new file mode 100644
--- /dev/null
+++ b/Engine/Plugins/Experimental/Fracture/Source/FractureEngine/Private/FractureToolClustering.cpp
@@ -0,0 +1,2 @@
+// Copyright Epic Games, Inc. All Rights Reserved.
+
diff --git a/Engine/Plugins/Experimental/Fracture/Source/FractureEngine/Private/FractureToolSelection.cpp b/Engine/Plugins/Experimental/Fracture/Source/FractureEngine/Private/FractureToolSelection.cpp
new file mode 100644
--- /dev/null
+++ b/Engine/Plugins/Experimental/Fracture/Source/FractureEngine/Private/FractureToolSelection.cpp
@@ -0,0 +1,2 @@
+// Copyright Epic Games, Inc. All Rights Reserved.
+
diff --git a/Engine/Plugins/Experimental/Fracture/Source/FractureEngine/Public/FractureToolSelection.h b/Engine/Plugins/Experimental/Fracture/Source/FractureEngine/Public/FractureToolSelection.h
new file mode 100644
--- /dev/null
+++ b/Engine/Plugins/Experimental/Fracture/Source/FractureEngine/Public/FractureToolSelection.h
@@ -0,0 +1,2 @@
+// Copyright Epic Games, Inc. All Rights Reserved.
+
