Dataflow:
- Renaming file(s)

diff --git a/Engine/Plugins/Experimental/Fracture/Source/FractureEngine/Private/FractureEngineClustering.cpp b/Engine/Plugins/Experimental/Fracture/Source/FractureEngine/Private/FractureEngineClustering.cpp
new file mode 100644
--- /dev/null
+++ b/Engine/Plugins/Experimental/Fracture/Source/FractureEngine/Private/FractureEngineClustering.cpp
@@ -0,0 +1,2 @@
+// Copyright Epic Games, Inc. All Rights Reserved.
+
diff --git a/Engine/Plugins/Experimental/Fracture/Source/FractureEngine/Private/FractureEngineSelection.cpp b/Engine/Plugins/Experimental/Fracture/Source/FractureEngine/Private/FractureEngineSelection.cpp
new file mode 100644
--- /dev/null
+++ b/Engine/Plugins/Experimental/Fracture/Source/FractureEngine/Private/FractureEngineSelection.cpp
@@ -0,0 +1,2 @@
+// Copyright Epic Games, Inc. All Rights Reserved.
+
diff --git a/Engine/Plugins/Experimental/Fracture/Source/FractureEngine/Private/FractureToolClustering.cpp b/Engine/Plugins/Experimental/Fracture/Source/FractureEngine/Private/FractureToolClustering.cpp
deleted file mode 100644
diff --git a/Engine/Plugins/Experimental/Fracture/Source/FractureEngine/Private/FractureToolSelection.cpp b/Engine/Plugins/Experimental/Fracture/Source/FractureEngine/Private/FractureToolSelection.cpp
deleted file mode 100644
diff --git a/Engine/Plugins/Experimental/Fracture/Source/FractureEngine/Public/FractureEngineClustering.h b/Engine/Plugins/Experimental/Fracture/Source/FractureEngine/Public/FractureEngineClustering.h
new file mode 100644
--- /dev/null
+++ b/Engine/Plugins/Experimental/Fracture/Source/FractureEngine/Public/FractureEngineClustering.h
@@ -0,0 +1,2 @@
+// Copyright Epic Games, Inc. All Rights Reserved.
+
diff --git a/Engine/Plugins/Experimental/Fracture/Source/FractureEngine/Public/FractureEngineSelection.h b/Engine/Plugins/Experimental/Fracture/Source/FractureEngine/Public/FractureEngineSelection.h
new file mode 100644
--- /dev/null
+++ b/Engine/Plugins/Experimental/Fracture/Source/FractureEngine/Public/FractureEngineSelection.h
@@ -0,0 +1,2 @@
+// Copyright Epic Games, Inc. All Rights Reserved.
+
diff --git a/Engine/Plugins/Experimental/Fracture/Source/FractureEngine/Public/FractureToolSelection.h b/Engine/Plugins/Experimental/Fracture/Source/FractureEngine/Public/FractureToolSelection.h
deleted file mode 100644
