Fracture editor : add parameter to enforce more strictly the number of cluster when using the connectivity data
- Add parameter to tool and dataflow
- Change logic to loop and find the optimal auto cluster parameter

diff --git a/Engine/Plugins/Experimental/ChaosEditor/Source/FractureEditor/Private/FractureToolAutoCluster.cpp b/Engine/Plugins/Experimental/ChaosEditor/Source/FractureEditor/Private/FractureToolAutoCluster.cpp
--- a/Engine/Plugins/Experimental/ChaosEditor/Source/FractureEditor/Private/FractureToolAutoCluster.cpp
+++ b/Engine/Plugins/Experimental/ChaosEditor/Source/FractureEditor/Private/FractureToolAutoCluster.cpp
@@ -134,7 +134,8 @@ void UFractureToolAutoCluster::Execute(TWeakPtr<FFractureEditorModeToolkit> InTo
 						AutoClusterSettings->SiteCountFraction,
 						AutoClusterSettings->SiteSize,
 						AutoClusterSettings->bEnforceConnectivity,
-						AutoClusterSettings->bAvoidIsolated);
+						AutoClusterSettings->bAvoidIsolated,
+						AutoClusterSettings->bEnforceSiteParameters);
 				}
 
 				Context.GenerateGuids(StartTransformCount);
diff --git a/Engine/Plugins/Experimental/ChaosEditor/Source/FractureEditor/Private/FractureToolAutoCluster.h b/Engine/Plugins/Experimental/ChaosEditor/Source/FractureEditor/Private/FractureToolAutoCluster.h
--- a/Engine/Plugins/Experimental/ChaosEditor/Source/FractureEditor/Private/FractureToolAutoCluster.h
+++ b/Engine/Plugins/Experimental/ChaosEditor/Source/FractureEditor/Private/FractureToolAutoCluster.h
@@ -69,6 +69,10 @@ public:
 	UPROPERTY(EditAnywhere, Category = AutoCluster, meta = (DisplayName = "Enforce Cluster Connectivity"))
 	bool bEnforceConnectivity=true;
 
+	/** If true, make sure the site parameters are matched as close as possible ( bEnforceConnectivity can make the number of site larger than the requested input may produce without it ) */
+	UPROPERTY(EditAnywhere, Category = AutoCluster, meta = (EditCondition = "bEnforceConnectivity == true"))
+	bool bEnforceSiteParameters = true;
+
 	/** If true, prevent the creation of clusters with only a single child. Either by merging into a neighboring cluster, or not creating the cluster. */
 	UPROPERTY(EditAnywhere, Category = AutoCluster)
 	bool bAvoidIsolated = true;
diff --git a/Engine/Plugins/Experimental/Fracture/Source/FractureEngine/Private/FractureEngineClustering.cpp b/Engine/Plugins/Experimental/Fracture/Source/FractureEngine/Private/FractureEngineClustering.cpp
--- a/Engine/Plugins/Experimental/Fracture/Source/FractureEngine/Private/FractureEngineClustering.cpp
+++ b/Engine/Plugins/Experimental/Fracture/Source/FractureEngine/Private/FractureEngineClustering.cpp
@@ -436,7 +436,8 @@ void FFractureEngineClustering::AutoCluster(FGeometryCollection& GeometryCollect
 	const float SiteCountFraction,
 	const float SiteSize,
 	const bool bEnforceConnectivity,
-	const bool bAvoidIsolated)
+	const bool bAvoidIsolated,
+	const bool bEnforceSiteParameters)
 {
 	FFractureEngineBoneSelection Selection(GeometryCollection, BoneIndices);
 
@@ -444,7 +445,7 @@ void FFractureEngineClustering::AutoCluster(FGeometryCollection& GeometryCollect
 
 	for (const int32 ClusterIndex : Selection.GetSelectedBones())
 	{
-		AutoCluster(GeometryCollection, ClusterIndex, ClusterSizeMethod, SiteCount, SiteCountFraction, SiteSize, bEnforceConnectivity, bAvoidIsolated);
+		AutoCluster(GeometryCollection, ClusterIndex, ClusterSizeMethod, SiteCount, SiteCountFraction, SiteSize, bEnforceConnectivity, bAvoidIsolated, bEnforceSiteParameters);
 	}
 }
 
@@ -455,50 +456,67 @@ void FFractureEngineClustering::AutoCluster(FGeometryCollection& GeometryCollect
 	const float SiteCountFraction,
 	const float SiteSize,
 	const bool bEnforceConnectivity,
-	const bool bAvoidIsolated)
+	const bool bAvoidIsolated,
+	const bool bEnforceSiteParameters)
 {
 	FVoronoiPartitioner VoronoiPartition(&GeometryCollection, ClusterIndex);
 	int32 NumChildren = GeometryCollection.Children[ClusterIndex].Num();
 
-	int32 SiteCountToUse = 1;
+	int32 DesiredSiteCountToUse = 1;
 	if (ClusterSizeMethod == EFractureEngineClusterSizeMethod::ByNumber)
 	{
-		SiteCountToUse = SiteCount;
+		DesiredSiteCountToUse = SiteCount;
 	}
 	else if (ClusterSizeMethod == EFractureEngineClusterSizeMethod::ByFractionOfInput)
 	{
-		SiteCountToUse = FMath::Max(2, NumChildren * SiteCountFraction);
+		DesiredSiteCountToUse = FMath::Max(2, NumChildren * SiteCountFraction);
 	}
 	else if (ClusterSizeMethod == EFractureEngineClusterSizeMethod::BySize)
 	{
 		float TotalVolume = GeometryCollection.GetBoundingBox().GetBox().GetVolume();
 		float DesiredVolume = FMath::Pow(SiteSize, 3);
-		SiteCountToUse = FMath::Max(1, TotalVolume / DesiredVolume);
+		DesiredSiteCountToUse = FMath::Max(1, TotalVolume / DesiredVolume);
 	}
 
-	VoronoiPartition.KMeansPartition(SiteCountToUse);
-	if (VoronoiPartition.GetPartitionCount() == 0)
+	int32 SiteCountToUse = DesiredSiteCountToUse;
+	int32 PreviousPartitionCount = TNumericLimits<int32>::Max();
+	bool bIterate = false;
+	do
 	{
-		return;
-	}
+		VoronoiPartition.KMeansPartition(SiteCountToUse);
+		if (VoronoiPartition.GetPartitionCount() == 0)
+		{
+			return;
+		}
 
-	bool bNeedProximity = bEnforceConnectivity || bAvoidIsolated;
-	if (bNeedProximity)
-	{
-		FGeometryCollectionProximityUtility ProximityUtility(&GeometryCollection);
-		ProximityUtility.UpdateProximity();
-	}
+		bool bNeedProximity = bEnforceConnectivity || bAvoidIsolated;
+		if (bNeedProximity)
+		{
+			FGeometryCollectionProximityUtility ProximityUtility(&GeometryCollection);
+			ProximityUtility.UpdateProximity();
+		}
 
-	if (bEnforceConnectivity)
-	{
-		VoronoiPartition.SplitDisconnectedPartitions(&GeometryCollection);
-	}
+		if (bEnforceConnectivity)
+		{
+			VoronoiPartition.SplitDisconnectedPartitions(&GeometryCollection);
+		}
 
-	if (bAvoidIsolated)
-	{
-		// attempt to remove isolated via merging (may not succeed, as it only merges if there is a cluster in proximity)
-		VoronoiPartition.MergeSingleElementPartitions(&GeometryCollection);
-	}
+		if (bAvoidIsolated)
+		{
+			// attempt to remove isolated via merging (may not succeed, as it only merges if there is a cluster in proximity)
+			VoronoiPartition.MergeSingleElementPartitions(&GeometryCollection);
+		}
+		const int32 IsolatedPartitionToIgnore = bAvoidIsolated ? VoronoiPartition.GetIsolatedPartitionCount() : 0;
+		const int32 PartitionCount = VoronoiPartition.GetNonEmptyPartitionCount() - IsolatedPartitionToIgnore;
+		bIterate = bEnforceSiteParameters				 // Is the feature enabled?
+			&& (PartitionCount > DesiredSiteCountToUse)  // Have we reached the desired outcome
+			&& (SiteCountToUse > 1)						 // Is SiteCount large enough ?
+			&& (PartitionCount < PreviousPartitionCount);// Are we progressing ?
+		if(bIterate)
+		{
+			SiteCountToUse--;
+		}
+	} while (bIterate);
 
 	int32 NonEmptyPartitionCount = VoronoiPartition.GetNonEmptyPartitionCount();
 	bool bHasEmptyClusters = NonEmptyPartitionCount > 0;
@@ -539,4 +557,3 @@ void FFractureEngineClustering::AutoCluster(FGeometryCollection& GeometryCollect
 }
 
 
-
diff --git a/Engine/Plugins/Experimental/Fracture/Source/FractureEngine/Public/FractureEngineClustering.h b/Engine/Plugins/Experimental/Fracture/Source/FractureEngine/Public/FractureEngineClustering.h
--- a/Engine/Plugins/Experimental/Fracture/Source/FractureEngine/Public/FractureEngineClustering.h
+++ b/Engine/Plugins/Experimental/Fracture/Source/FractureEngine/Public/FractureEngineClustering.h
@@ -29,6 +29,11 @@ public:
 		return PartitionSize.Num() - Algo::Count(PartitionSize, 0);
 	}
 
+	int32 GetIsolatedPartitionCount() const
+	{
+		return Algo::Count(PartitionSize, 1);
+	}
+
 	/** return the GeometryCollection TranformIndices within the partition. */
 	TArray<int32> GetPartition(int32 PartitionIndex) const;
 
@@ -79,7 +84,8 @@ public:
 		const float SiteCountFraction,
 		const float SiteSize,
 		const bool bEnforceConnectivity,
-		const bool bAvoidIsolated);
+		const bool bAvoidIsolated,
+		const bool bEnforceSiteParameters);
 
 	static void AutoCluster(FGeometryCollection& GeometryCollection,
 		const int32 ClusterIndex,
@@ -88,7 +94,8 @@ public:
 		const float SiteCountFraction,
 		const float SiteSize,
 		const bool bEnforceConnectivity,
-		const bool bAvoidIsolated);
+		const bool bAvoidIsolated,
+		const bool bEnforceSiteParameters);
 };
 
 #if UE_ENABLE_INCLUDE_ORDER_DEPRECATED_IN_5_2
diff --git a/Engine/Plugins/Experimental/GeometryCollectionPlugin/Source/GeometryCollectionNodes/Private/Dataflow/GeometryCollectionClusteringNodes.cpp b/Engine/Plugins/Experimental/GeometryCollectionPlugin/Source/GeometryCollectionNodes/Private/Dataflow/GeometryCollectionClusteringNodes.cpp
--- a/Engine/Plugins/Experimental/GeometryCollectionPlugin/Source/GeometryCollectionNodes/Private/Dataflow/GeometryCollectionClusteringNodes.cpp
+++ b/Engine/Plugins/Experimental/GeometryCollectionPlugin/Source/GeometryCollectionNodes/Private/Dataflow/GeometryCollectionClusteringNodes.cpp
@@ -60,6 +60,7 @@ void FAutoClusterDataflowNode::Evaluate(Dataflow::FContext& Context, const FData
 			float InClusterFraction = GetValue<float>(Context, &ClusterFraction);
 			float InSiteSize = GetValue<float>(Context, &SiteSize);
 			bool InAutoCluster = AutoCluster;
+			bool InEnforceSiteParameters = EnforceSiteParameters;
 			bool InAvoidIsolated = AvoidIsolated;
 
 			TArray<int32> SelectedBones;
@@ -72,7 +73,8 @@ void FAutoClusterDataflowNode::Evaluate(Dataflow::FContext& Context, const FData
 				InClusterFraction,
 				InSiteSize,
 				InAutoCluster,
-				InAvoidIsolated);
+				InAvoidIsolated, 
+				InEnforceSiteParameters);
 
 			SetValue<FManagedArrayCollection>(Context, (const FManagedArrayCollection&)(*GeomCollection), &Collection);
 		}
diff --git a/Engine/Plugins/Experimental/GeometryCollectionPlugin/Source/GeometryCollectionNodes/Public/Dataflow/GeometryCollectionClusteringNodes.h b/Engine/Plugins/Experimental/GeometryCollectionPlugin/Source/GeometryCollectionNodes/Public/Dataflow/GeometryCollectionClusteringNodes.h
--- a/Engine/Plugins/Experimental/GeometryCollectionPlugin/Source/GeometryCollectionNodes/Public/Dataflow/GeometryCollectionClusteringNodes.h
+++ b/Engine/Plugins/Experimental/GeometryCollectionPlugin/Source/GeometryCollectionNodes/Public/Dataflow/GeometryCollectionClusteringNodes.h
@@ -55,6 +55,10 @@ public:
 	UPROPERTY(EditAnywhere, Category = "Auto Cluster")
 	bool AutoCluster = true;
 
+	/** If true, make sure the site parameters are matched as close as possible ( bEnforceConnectivity can make the number of site larger than the requested input may produce without it ) */
+	UPROPERTY(EditAnywhere, Category = AutoCluster, meta = (EditCondition = "bEnforceConnectivity == true"))
+	bool EnforceSiteParameters = true;
+
 	/** If true, prevent the creation of clusters with only a single child. Either by merging into a neighboring cluster, or not creating the cluster. */
 	UPROPERTY(EditAnywhere, Category = "Auto Cluster")
 	bool AvoidIsolated = true;
