Expand Fracture AutoCluster tool:
- Add grid-based clustering options
- Add a minimum cluster volume option
- Add control over the K-Means max iterations (to help allow clusters that follow the grid better)

diff --git a/Engine/Plugins/Experimental/ChaosEditor/Source/FractureEditor/Private/FractureToolAutoCluster.cpp b/Engine/Plugins/Experimental/ChaosEditor/Source/FractureEditor/Private/FractureToolAutoCluster.cpp
--- a/Engine/Plugins/Experimental/ChaosEditor/Source/FractureEditor/Private/FractureToolAutoCluster.cpp
+++ b/Engine/Plugins/Experimental/ChaosEditor/Source/FractureEditor/Private/FractureToolAutoCluster.cpp
@@ -135,7 +135,12 @@ void UFractureToolAutoCluster::Execute(TWeakPtr<FFractureEditorModeToolkit> InTo
 						AutoClusterSettings->SiteSize,
 						AutoClusterSettings->bEnforceConnectivity,
 						AutoClusterSettings->bAvoidIsolated,
-						AutoClusterSettings->bEnforceSiteParameters);
+						AutoClusterSettings->bEnforceSiteParameters,
+						AutoClusterSettings->ClusterGridWidth,
+						AutoClusterSettings->ClusterGridDepth,
+						AutoClusterSettings->ClusterGridHeight,
+						AutoClusterSettings->MinimumSize,
+						AutoClusterSettings->KMeansIterations);
 				}
 
 				Context.GenerateGuids(StartTransformCount);
diff --git a/Engine/Plugins/Experimental/ChaosEditor/Source/FractureEditor/Private/FractureToolAutoCluster.h b/Engine/Plugins/Experimental/ChaosEditor/Source/FractureEditor/Private/FractureToolAutoCluster.h
--- a/Engine/Plugins/Experimental/ChaosEditor/Source/FractureEditor/Private/FractureToolAutoCluster.h
+++ b/Engine/Plugins/Experimental/ChaosEditor/Source/FractureEditor/Private/FractureToolAutoCluster.h
@@ -32,6 +32,8 @@ enum class EClusterSizeMethod : uint8
 	ByFractionOfInput,
 	// Cluster by specifying the density of the input bones
 	BySize,
+	// Cluster by a grid layout
+	ByGrid,
 };
 
 
@@ -43,7 +45,6 @@ public:
 
 	UFractureAutoClusterSettings(const FObjectInitializer& ObjInit)
 		: Super(ObjInit)
-		, SiteCount(10)
 	{}
 
 	UPROPERTY(meta = (DeprecatedProperty, DeprecationMessage = "Simplified interface now only supports Voronoi clustering."))
@@ -54,23 +55,43 @@ public:
 	EClusterSizeMethod ClusterSizeMethod = EClusterSizeMethod::ByNumber;
 
 	/** Use a Voronoi diagram with this many Voronoi sites as a guide for deciding cluster boundaries */
-	UPROPERTY(EditAnywhere, Category = ClusterSize, meta = (DisplayName = "Cluster Sites", UIMin = "2", UIMax = "5000", ClampMin = "1", EditCondition = "ClusterSizeMethod == EClusterSizeMethod::ByNumber"))
+	UPROPERTY(EditAnywhere, Category = ClusterSize, meta = (DisplayName = "Cluster Sites", UIMin = "2", UIMax = "5000", ClampMin = "1", EditCondition = "ClusterSizeMethod == EClusterSizeMethod::ByNumber", EditConditionHides))
 	uint32 SiteCount=10;
 
 	/** Choose the number of Voronoi sites used for clustering as a fraction of the number of child bones to process */
-	UPROPERTY(EditAnywhere, Category = ClusterSize, meta = (DisplayName = "Cluster Fraction", ClampMin = "0", ClampMax = ".5", EditCondition = "ClusterSizeMethod == EClusterSizeMethod::ByFractionOfInput"))
+	UPROPERTY(EditAnywhere, Category = ClusterSize, meta = (DisplayName = "Cluster Fraction", ClampMin = "0", ClampMax = ".5", EditCondition = "ClusterSizeMethod == EClusterSizeMethod::ByFractionOfInput", EditConditionHides))
 	float SiteCountFraction = .25;
 	
 	/** Choose the Edge-Size of the cube used to groups bones under a cluster (in cm). */
-	UPROPERTY(EditAnywhere, Category = ClusterSize, meta = (DisplayName = "Cluster Size", UIMin = ".01", UIMax = "100", ClampMin = ".0001", ClampMax = "10000", EditCondition = "ClusterSizeMethod == EClusterSizeMethod::BySize"))
+	UPROPERTY(EditAnywhere, Category = ClusterSize, meta = (DisplayName = "Cluster Size", UIMin = ".01", UIMax = "100", ClampMin = ".0001", ClampMax = "10000", EditCondition = "ClusterSizeMethod == EClusterSizeMethod::BySize", EditConditionHides))
 	float SiteSize = 1;
 
+	/** Choose the number of cluster sites to distribute along the X axis */
+	UPROPERTY(EditAnywhere, Category = ClusterSize, meta = (ClampMin = "1", UIMax = "20", EditCondition = "ClusterSizeMethod == EClusterSizeMethod::ByGrid", EditConditionHides))
+	int ClusterGridWidth = 2;
+
+	/** Choose the number of cluster sites to distribute along the Y axis */
+	UPROPERTY(EditAnywhere, Category = ClusterSize, meta = (ClampMin = "1", UIMax = "20", EditCondition = "ClusterSizeMethod == EClusterSizeMethod::ByGrid", EditConditionHides))
+	int ClusterGridDepth = 2;
+	
+	/** Choose the number of cluster sites to distribute along the Z axis */
+	UPROPERTY(EditAnywhere, Category = ClusterSize, meta = (ClampMin = "1", UIMax = "20", EditCondition = "ClusterSizeMethod == EClusterSizeMethod::ByGrid", EditConditionHides))
+	int ClusterGridHeight = 2;
+
+	/** If a cluster has volume less than this value (in cm) cubed, then the auto-cluster process will attempt to merge it into a neighboring cluster. */
+	UPROPERTY(EditAnywhere, Category = ClusterSize, meta = (ClampMin = "0"))
+	float MinimumSize = 0;
+
+	/** The maximum number of K-Means iterations to apply, balancing the shape and distribution of the clusters */
+	UPROPERTY(EditAnywhere, Category = ClusterSize, meta = (DisplayName = "K-Means Iteration Max", ClampMin = "1", UIMax = "500"))
+	int KMeansIterations = 500;
+
 	/** If true, bones will only be added to the same cluster if they are physically connected (either directly, or via other bones in the same cluster) */
 	UPROPERTY(EditAnywhere, Category = AutoCluster, meta = (DisplayName = "Enforce Cluster Connectivity"))
 	bool bEnforceConnectivity=true;
 
 	/** If true, make sure the site parameters are matched as close as possible ( bEnforceConnectivity can make the number of site larger than the requested input may produce without it ) */
-	UPROPERTY(EditAnywhere, Category = AutoCluster, meta = (EditCondition = "bEnforceConnectivity == true"))
+	UPROPERTY(EditAnywhere, Category = AutoCluster, meta = (EditCondition = "bEnforceConnectivity == true && ClusterSizeMethod != EClusterSizeMethod::ByGrid"))
 	bool bEnforceSiteParameters = true;
 
 	/** If true, prevent the creation of clusters with only a single child. Either by merging into a neighboring cluster, or not creating the cluster. */
diff --git a/Engine/Plugins/Experimental/Fracture/Source/FractureEngine/Private/FractureEngineClustering.cpp b/Engine/Plugins/Experimental/Fracture/Source/FractureEngine/Private/FractureEngineClustering.cpp
--- a/Engine/Plugins/Experimental/Fracture/Source/FractureEngine/Private/FractureEngineClustering.cpp
+++ b/Engine/Plugins/Experimental/Fracture/Source/FractureEngine/Private/FractureEngineClustering.cpp
@@ -15,12 +15,16 @@ FVoronoiPartitioner::FVoronoiPartitioner(const FGeometryCollection* GeometryColl
 	GenerateCentroids(GeometryCollection);
 }
 
-void FVoronoiPartitioner::KMeansPartition(int32 InPartitionCount)
+void FVoronoiPartitioner::KMeansPartition(int32 InPartitionCount, int32 MaxIterations, TArrayView<const FVector> InitialCenters)
 {
 	PartitionCount = InPartitionCount;
+	if (InitialCenters.Num())
+	{
+		PartitionCount = InitialCenters.Num();
+	}
 
-	InitializePartitions();
-	for (int32 Iteration = 0; Iteration < MaxKMeansIterations; Iteration++)
+	InitializePartitions(InitialCenters);
+	for (int32 Iteration = 0; Iteration < MaxIterations; Iteration++)
 	{
 		// Refinement is complete when no nodes change partition.
 		if (!Refine())
@@ -69,6 +73,93 @@ void FVoronoiPartitioner::MergeSingleElementPartitions(FGeometryCollection* Geom
 	}
 }
 
+void FVoronoiPartitioner::MergeSmallPartitions(FGeometryCollection* GeometryCollection, float SizeThreshold)
+{
+	if (Connectivity.IsEmpty())
+	{
+		GenerateConnectivity(GeometryCollection);
+	}
+	if (!GeometryCollection->HasAttribute("Volume", FTransformCollection::TransformGroup))
+	{
+		FGeometryCollectionConvexUtility::SetVolumeAttributes(GeometryCollection);
+	}
+	const TManagedArray<float>& Volumes = GeometryCollection->GetAttribute<float>("Volume", FTransformCollection::TransformGroup);
+
+	float VolumeThreshold = SizeThreshold * SizeThreshold * SizeThreshold;
+	TArray<float> PartitionVolumes;
+	PartitionVolumes.Init(0, PartitionSize.Num());
+	
+	int32 NonEmptyPartitions = GetNonEmptyPartitionCount();
+
+	for (int32 ElIdx = 0; ElIdx < TransformIndices.Num(); ElIdx++)
+	{
+		int32 Partition = Partitions[ElIdx];
+		PartitionVolumes[Partition] += Volumes[TransformIndices[ElIdx]];
+	}
+
+	TArray<int32> ToConsider;
+	ToConsider.Reserve(Partitions.Num());
+	for (int32 Partition = 0; Partition < PartitionVolumes.Num(); ++Partition)
+	{
+		if (PartitionVolumes[Partition] < VolumeThreshold)
+		{
+			ToConsider.Emplace(Partition);
+		}
+	}
+	while (!ToConsider.IsEmpty())
+	{
+		int32 Partition = ToConsider.Pop(false);
+		if (NonEmptyPartitions < 3)
+		{
+			break; // if we only have two partitions, stop merging
+		}
+		if (PartitionVolumes[Partition] < VolumeThreshold)
+		{
+			int32 SmallestNbrPartition = INDEX_NONE;
+			float SmallestNbrVolume = FMathf::MaxReal;
+			for (int32 ElIdx = 0; ElIdx < TransformIndices.Num(); ++ElIdx)
+			{
+				if (Partitions[ElIdx] == Partition)
+				{
+					for (int32 NbrEl : Connectivity[ElIdx])
+					{
+						int32 NbrPartition = Partitions[NbrEl];
+						if (NbrPartition != Partition)
+						{
+							if (SmallestNbrPartition == INDEX_NONE || PartitionVolumes[NbrPartition] < SmallestNbrVolume)
+							{
+								SmallestNbrVolume = PartitionVolumes[NbrPartition];
+								SmallestNbrPartition = NbrPartition;
+							}
+						}
+					}
+				}
+			}
+			if (SmallestNbrPartition != INDEX_NONE)
+			{
+				PartitionVolumes[SmallestNbrPartition] += PartitionVolumes[Partition];
+				PartitionVolumes[Partition] = 0;
+				float CombinedSize = PartitionSize[SmallestNbrPartition] + PartitionSize[Partition];
+				PartitionSize[SmallestNbrPartition] = CombinedSize;
+				if (CombinedSize < VolumeThreshold)
+				{
+					ToConsider.Push(SmallestNbrPartition);
+					Swap(ToConsider.Last(), ToConsider[0]); // make sure we don't always merge to the same cluster
+				}
+				PartitionSize[Partition] = 0;
+				NonEmptyPartitions--;
+				for (int32 ElIdx = 0; ElIdx < TransformIndices.Num(); ++ElIdx)
+				{
+					if (Partitions[ElIdx] == Partition)
+					{
+						Partitions[ElIdx] = SmallestNbrPartition;
+					}
+				}
+			}
+		}
+	}
+}
+
 void FVoronoiPartitioner::SplitDisconnectedPartitions(FGeometryCollection* GeometryCollection)
 {
 	Visited.Init(false, TransformIndices.Num());
@@ -195,7 +286,7 @@ FVector FVoronoiPartitioner::GenerateCentroid(const FGeometryCollection* Geometr
 	return Bounds.GetCenter();
 }
 
-FBox FVoronoiPartitioner::GenerateBounds(const FGeometryCollection* GeometryCollection, int32 TransformIndex) const
+FBox FVoronoiPartitioner::GenerateBounds(const FGeometryCollection* GeometryCollection, int32 TransformIndex)
 {
 	check(GeometryCollection->IsRigid(TransformIndex) || GeometryCollection->IsClustered(TransformIndex));
 
@@ -241,18 +332,26 @@ FBox FVoronoiPartitioner::GenerateBounds(const FGeometryCollection* GeometryColl
 	}
 }
 
-void FVoronoiPartitioner::InitializePartitions()
+void FVoronoiPartitioner::InitializePartitions(TArrayView<const FVector> InitialCenters)
 {
 	check(PartitionCount > 0);
 
-	PartitionCount = FMath::Min(PartitionCount, TransformIndices.Num());
-
-	// Set initial partition centers as selects from the vertex set
-	PartitionCenters.SetNum(PartitionCount);
-	int32 TransformStride = FMath::Max(1,FMath::FloorToInt(float(TransformIndices.Num()) / float(PartitionCount)));
-	for (int32 PartitionIndex = 0; PartitionIndex < PartitionCount; ++PartitionIndex)
+	if (InitialCenters.Num() > 0)
 	{
-		PartitionCenters[PartitionIndex] = Centroids[TransformStride * PartitionIndex];
+		PartitionCenters.Reset();
+		PartitionCenters.Append(InitialCenters);
+	}
+	else
+	{
+		PartitionCount = FMath::Min(PartitionCount, TransformIndices.Num());
+
+		// Set initial partition centers as selects from the vertex set
+		PartitionCenters.SetNum(PartitionCount);
+		int32 TransformStride = FMath::Max(1, FMath::FloorToInt(float(TransformIndices.Num()) / float(PartitionCount)));
+		for (int32 PartitionIndex = 0; PartitionIndex < PartitionCount; ++PartitionIndex)
+		{
+			PartitionCenters[PartitionIndex] = Centroids[TransformStride * PartitionIndex];
+		}
 	}
 
 	// At beginning, all nodes belong to first partition.
@@ -437,7 +536,12 @@ void FFractureEngineClustering::AutoCluster(FGeometryCollection& GeometryCollect
 	const float SiteSize,
 	const bool bEnforceConnectivity,
 	const bool bAvoidIsolated,
-	const bool bEnforceSiteParameters)
+	const bool bEnforceSiteParameters,
+	const int32 GridX,
+	const int32 GridY,
+	const int32 GridZ,
+	const float MinimumClusterSize,
+	const int32 KMeansIterations)
 {
 	FFractureEngineBoneSelection Selection(GeometryCollection, BoneIndices);
 
@@ -445,7 +549,8 @@ void FFractureEngineClustering::AutoCluster(FGeometryCollection& GeometryCollect
 
 	for (const int32 ClusterIndex : Selection.GetSelectedBones())
 	{
-		AutoCluster(GeometryCollection, ClusterIndex, ClusterSizeMethod, SiteCount, SiteCountFraction, SiteSize, bEnforceConnectivity, bAvoidIsolated, bEnforceSiteParameters);
+		AutoCluster(GeometryCollection, ClusterIndex, ClusterSizeMethod, SiteCount, SiteCountFraction, SiteSize, 
+			bEnforceConnectivity, bAvoidIsolated, bEnforceSiteParameters, GridX, GridY, GridZ, MinimumClusterSize, KMeansIterations);
 	}
 }
 
@@ -457,12 +562,21 @@ void FFractureEngineClustering::AutoCluster(FGeometryCollection& GeometryCollect
 	const float SiteSize,
 	const bool bEnforceConnectivity,
 	const bool bAvoidIsolated,
-	const bool bEnforceSiteParameters)
+	const bool bEnforceSiteParameters,
+	const int32 InGridX,
+	const int32 InGridY,
+	const int32 InGridZ,
+	const float MinimumClusterSize,
+	const int32 KMeansIterations)
 {
 	FVoronoiPartitioner VoronoiPartition(&GeometryCollection, ClusterIndex);
 	int32 NumChildren = GeometryCollection.Children[ClusterIndex].Num();
 
+	// Used by the ByGrid method, which manually distributes initial positions
+	TArray<FVector> PartitionPositions;
+
 	int32 DesiredSiteCountToUse = 1;
+	int32 IterationsToUse = KMeansIterations;
 	if (ClusterSizeMethod == EFractureEngineClusterSizeMethod::ByNumber)
 	{
 		DesiredSiteCountToUse = SiteCount;
@@ -477,13 +591,38 @@ void FFractureEngineClustering::AutoCluster(FGeometryCollection& GeometryCollect
 		float DesiredVolume = FMath::Pow(SiteSize, 3);
 		DesiredSiteCountToUse = FMath::Max(1, TotalVolume / DesiredVolume);
 	}
+	else if (ClusterSizeMethod == EFractureEngineClusterSizeMethod::ByGrid)
+	{
+		FBox Bounds = FVoronoiPartitioner::GenerateBounds(&GeometryCollection, ClusterIndex);
+		int32 GridX = FMath::Max(1, InGridX);
+		int32 GridY = FMath::Max(1, InGridY);
+		int32 GridZ = FMath::Max(1, InGridZ);
+		PartitionPositions.Reserve(GridX * GridY * GridZ);
+
+		for (int32 StepX = 0; StepX < GridX; ++StepX)
+		{
+			double Xt = GridX == 1 ? .5f : double(StepX) / double(GridX - 1);
+			double X = FMath::Lerp(Bounds.Min.X, Bounds.Max.X, Xt);
+			for (int32 StepY = 0; StepY < GridY; ++StepY)
+			{
+				double Yt = GridY == 1 ? .5f : double(StepY) / double(GridY - 1);
+				double Y = FMath::Lerp(Bounds.Min.Y, Bounds.Max.Y, Yt);
+				for (int32 StepZ = 0; StepZ < GridZ; ++StepZ)
+				{
+					double Zt = GridZ == 1 ? .5f : double(StepZ) / double(GridZ - 1);
+					double Z = FMath::Lerp(Bounds.Min.Z, Bounds.Max.Z, Zt);
+					PartitionPositions.Emplace(X, Y, Z);
+				}
+			}
+		}
+	}
 
 	int32 SiteCountToUse = DesiredSiteCountToUse;
 	int32 PreviousPartitionCount = TNumericLimits<int32>::Max();
 	bool bIterate = false;
 	do
 	{
-		VoronoiPartition.KMeansPartition(SiteCountToUse);
+		VoronoiPartition.KMeansPartition(SiteCountToUse, IterationsToUse, PartitionPositions);
 		if (VoronoiPartition.GetPartitionCount() == 0)
 		{
 			return;
@@ -506,9 +645,16 @@ void FFractureEngineClustering::AutoCluster(FGeometryCollection& GeometryCollect
 			// attempt to remove isolated via merging (may not succeed, as it only merges if there is a cluster in proximity)
 			VoronoiPartition.MergeSingleElementPartitions(&GeometryCollection);
 		}
+
+		if (MinimumClusterSize > 0)
+		{
+			// attempt to remove isolated via cluster merging (may not succeed, as it only merges if there is a cluster in proximity)
+			VoronoiPartition.MergeSmallPartitions(&GeometryCollection, MinimumClusterSize);
+		}
 		const int32 IsolatedPartitionToIgnore = bAvoidIsolated ? VoronoiPartition.GetIsolatedPartitionCount() : 0;
 		const int32 PartitionCount = VoronoiPartition.GetNonEmptyPartitionCount() - IsolatedPartitionToIgnore;
 		bIterate = bEnforceSiteParameters				 // Is the feature enabled?
+			&& PartitionPositions.IsEmpty()				 // Is the explicit position array empty? (otherwise, site count is ignored)
 			&& (PartitionCount > DesiredSiteCountToUse)  // Have we reached the desired outcome
 			&& (SiteCountToUse > 1)						 // Is SiteCount large enough ?
 			&& (PartitionCount < PreviousPartitionCount);// Are we progressing ?
diff --git a/Engine/Plugins/Experimental/Fracture/Source/FractureEngine/Public/FractureEngineClustering.h b/Engine/Plugins/Experimental/Fracture/Source/FractureEngine/Public/FractureEngineClustering.h
--- a/Engine/Plugins/Experimental/Fracture/Source/FractureEngine/Public/FractureEngineClustering.h
+++ b/Engine/Plugins/Experimental/Fracture/Source/FractureEngine/Public/FractureEngineClustering.h
@@ -13,8 +13,13 @@ class FRACTUREENGINE_API FVoronoiPartitioner
 public:
 	FVoronoiPartitioner(const FGeometryCollection* GeometryCollection, int32 ClusterIndex);
 
-	/** Cluster bodies into k partitions using K-Means. Connectivity is ignored: only spatial proximity is considered. */
-	void KMeansPartition(int32 InPartitionCount);
+	/** 
+	 * Cluster bodies into k partitions using K-Means. Connectivity is ignored: only spatial proximity is considered. 
+	 * @param InPartitionCount	Number of partitions to target, if InitialCenters is not provided
+	 * @param MaxIterations		Maximum iterations of refinement of partitions. In many cases, K-Means will converge and stop early if MaxIterations is large.
+	 * @param InitialCenters	If non-empty, these positions will be used to initialize the partition locations. The target partition count will then be the length of this array.
+	 */
+	void KMeansPartition(int32 InPartitionCount, int32 MaxIterations = 500, TArrayView<const FVector> InitialCenters = TArrayView<const FVector>());
 
 	/** Split any partition islands into their own partition. This will possbily increase number of partitions to exceed desired count. */
 	void SplitDisconnectedPartitions(FGeometryCollection* GeometryCollection);
@@ -22,6 +27,9 @@ public:
 	/** Merge any partitions w/ only 1 body into a connected, neighboring partition (if any).  This can decrease the number of partitions below the desired count. */
 	void MergeSingleElementPartitions(FGeometryCollection* GeometryCollection);
 
+	/** Merge any too-small partitions into a connected, neighboring partition (if any).  This can decrease the number of partitions below the desired count. */
+	void MergeSmallPartitions(FGeometryCollection* GeometryCollection, float PartitionSizeThreshold);
+
 	int32 GetPartitionCount() const { return PartitionCount; }
 
 	int32 GetNonEmptyPartitionCount() const
@@ -37,13 +45,14 @@ public:
 	/** return the GeometryCollection TranformIndices within the partition. */
 	TArray<int32> GetPartition(int32 PartitionIndex) const;
 
+	static FBox GenerateBounds(const FGeometryCollection* GeometryCollection, int32 TransformIndex);
+
 private:
 	void GenerateConnectivity(const FGeometryCollection* GeometryCollection);
 	void CollectConnections(const FGeometryCollection* GeometryCollection, int32 Index, int32 OperatingLevel, TSet<int32>& OutConnections) const;
 	void GenerateCentroids(const FGeometryCollection* GeometryCollection);
 	FVector GenerateCentroid(const FGeometryCollection* GeometryCollection, int32 TransformIndex) const;
-	FBox GenerateBounds(const FGeometryCollection* GeometryCollection, int32 TransformIndex) const;
-	void InitializePartitions();
+	void InitializePartitions(TArrayView<const FVector> InitialCenters = TArrayView<const FVector>());
 	bool Refine();
 	int32 FindClosestPartitionCenter(const FVector& Location) const;
 	void MarkVisited(int32 Index, int32 PartitionIndex);
@@ -51,16 +60,14 @@ private:
 private:
 	TArray<int32> TransformIndices;
 	TArray<FVector> Centroids;
+	// mapping from index into TransformIndices to partition number
 	TArray<int32> Partitions;
 	int32 PartitionCount;
 	TArray<int32> PartitionSize;
 	TArray<FVector> PartitionCenters;
+	// mapping from index into TransformIndices to the set of connected transforms (also via their index in TransformIndices)
 	TArray<TSet<int32>> Connectivity;
 	TArray<bool> Visited;
-
-
-	// Not generally necessary but this is a safety measure to prevent oscillating solves that never converge.
-	const int32 MaxKMeansIterations = 500;
 };
 
 enum class FRACTUREENGINE_API EFractureEngineClusterSizeMethod : uint8
@@ -71,6 +78,8 @@ enum class FRACTUREENGINE_API EFractureEngineClusterSizeMethod : uint8
 	ByFractionOfInput,
 	// Cluster by specifying the density of the input bones
 	BySize,
+	// Cluster by a regular grid distribution
+	ByGrid,
 };
 
 class FRACTUREENGINE_API FFractureEngineClustering
@@ -85,7 +94,12 @@ public:
 		const float SiteSize,
 		const bool bEnforceConnectivity,
 		const bool bAvoidIsolated,
-		const bool bEnforceSiteParameters);
+		const bool bEnforceSiteParameters,
+		const int32 GridX = 2,
+		const int32 GridY = 2,
+		const int32 GridZ = 2,
+		const float MinimumClusterSize = 0,
+		const int32 KMeansIterations = 500);
 
 	static void AutoCluster(FGeometryCollection& GeometryCollection,
 		const int32 ClusterIndex,
@@ -95,7 +109,12 @@ public:
 		const float SiteSize,
 		const bool bEnforceConnectivity,
 		const bool bAvoidIsolated,
-		const bool bEnforceSiteParameters);
+		const bool bEnforceSiteParameters,
+		const int32 GridX = 2,
+		const int32 GridY = 2,
+		const int32 GridZ = 2,
+		const float MinimumClusterSize = 0,
+		const int32 KMeansIterations = 500);
 };
 
 #if UE_ENABLE_INCLUDE_ORDER_DEPRECATED_IN_5_2
