Add visualization of grid points to Auto Cluster tool, and adjust grid distribution so that grid points are better-centered

diff --git a/Engine/Plugins/Experimental/ChaosEditor/Source/FractureEditor/Private/FractureToolAutoCluster.cpp b/Engine/Plugins/Experimental/ChaosEditor/Source/FractureEditor/Private/FractureToolAutoCluster.cpp
--- a/Engine/Plugins/Experimental/ChaosEditor/Source/FractureEditor/Private/FractureToolAutoCluster.cpp
+++ b/Engine/Plugins/Experimental/ChaosEditor/Source/FractureEditor/Private/FractureToolAutoCluster.cpp
@@ -89,16 +89,56 @@ void UFractureToolAutoCluster::DrawBox(FPrimitiveDrawInterface* PDI, FVector Cen
 
 void UFractureToolAutoCluster::Render(const FSceneView* View, FViewport* Viewport, FPrimitiveDrawInterface* PDI)
 {
-	if (AutoClusterSettings->ClusterSizeMethod != EClusterSizeMethod::BySize)
+	if (AutoClusterSettings->ClusterSizeMethod == EClusterSizeMethod::ByGrid)
 	{
-		return;
+		EnumerateVisualizationMapping(GridPointsMappings, ShowGridPoints.Num(),
+			[&](int32 Idx, FVector ExplodedVector)
+			{
+				FVector Pt = ShowGridPoints[Idx];
+				PDI->DrawPoint(Pt, FLinearColor(1.f, 0.f, 1.f), 5, SDPG_Foreground);
+			});
 	}
 
-	TArray<FFractureToolContext> FractureContexts = GetFractureToolContexts();
+	if (AutoClusterSettings->ClusterSizeMethod == EClusterSizeMethod::BySize)
+	{
+		TArray<FFractureToolContext> FractureContexts = GetFractureToolContexts();
+
+		for (FFractureToolContext& FractureContext : FractureContexts)
+		{
+			DrawBox(PDI, FractureContext.GetTransform().GetTranslation(), AutoClusterSettings->SiteSize);
+		}
+	}
+}
+
+void UFractureToolAutoCluster::FractureContextChanged()
+{
+	ClearVisualizations();
+
+	TArray<FFractureToolContext> Contexts = GetFractureToolContexts();
 
-	for (FFractureToolContext& FractureContext : FractureContexts)
+	for (FFractureToolContext& Context : Contexts)
 	{
-		DrawBox(PDI, FractureContext.GetTransform().GetTranslation(), AutoClusterSettings->SiteSize);
+		Context.ConvertSelectionToClusterNodes();
+
+		if (const FGeometryCollection* GeometryCollection = Context.GetGeometryCollection().Get())
+		{
+			int CollectionIdx = VisualizedCollections.Emplace(Context.GetGeometryCollectionComponent());
+			FTransform ContextTransform = Context.GetTransform();
+
+			for (const int32 ClusterIndex : Context.GetSelection())
+			{
+				GridPointsMappings.AddMapping(CollectionIdx, ClusterIndex, ShowGridPoints.Num());
+				TArray<FVector> LocalGridPoints = FFractureEngineClustering::GenerateGridSites(
+					*GeometryCollection, ClusterIndex,
+					AutoClusterSettings->ClusterGridWidth,
+					AutoClusterSettings->ClusterGridDepth,
+					AutoClusterSettings->ClusterGridHeight);
+				for (FVector LocalPt : LocalGridPoints)
+				{
+					ShowGridPoints.Add(ContextTransform.TransformPosition(LocalPt));
+				}
+			}
+		}
 	}
 }
 
@@ -127,6 +167,7 @@ void UFractureToolAutoCluster::Execute(TWeakPtr<FFractureEditorModeToolkit> InTo
 
 				for (const int32 ClusterIndex : Context.GetSelection())
 				{
+					int32 KMeansIterations = AutoClusterSettings->ClusterSizeMethod == EClusterSizeMethod::ByGrid ? AutoClusterSettings->DriftIterations : 500;
 					FFractureEngineClustering::AutoCluster(*GeometryCollection,
 						ClusterIndex,
 						(EFractureEngineClusterSizeMethod)AutoClusterSettings->ClusterSizeMethod,
@@ -140,7 +181,7 @@ void UFractureToolAutoCluster::Execute(TWeakPtr<FFractureEditorModeToolkit> InTo
 						AutoClusterSettings->ClusterGridDepth,
 						AutoClusterSettings->ClusterGridHeight,
 						AutoClusterSettings->MinimumSize,
-						AutoClusterSettings->KMeansIterations);
+						KMeansIterations);
 				}
 
 				Context.GenerateGuids(StartTransformCount);
diff --git a/Engine/Plugins/Experimental/ChaosEditor/Source/FractureEditor/Private/FractureToolAutoCluster.h b/Engine/Plugins/Experimental/ChaosEditor/Source/FractureEditor/Private/FractureToolAutoCluster.h
--- a/Engine/Plugins/Experimental/ChaosEditor/Source/FractureEditor/Private/FractureToolAutoCluster.h
+++ b/Engine/Plugins/Experimental/ChaosEditor/Source/FractureEditor/Private/FractureToolAutoCluster.h
@@ -82,9 +82,9 @@ public:
 	UPROPERTY(EditAnywhere, Category = ClusterSize, meta = (ClampMin = "0"))
 	float MinimumSize = 0;
 
-	/** The maximum number of K-Means iterations to apply, balancing the shape and distribution of the clusters */
-	UPROPERTY(EditAnywhere, Category = ClusterSize, meta = (DisplayName = "K-Means Iteration Max", ClampMin = "1", UIMax = "500"))
-	int KMeansIterations = 500;
+	/** For a grid distribution, optionally iteratively recenter the grid points to the center of the cluster geometry (technically: applying K-Means iterations) to balance the shape and distribution of the clusters */
+	UPROPERTY(EditAnywhere, Category = ClusterSize, meta = (DisplayName = "Drift Iterations", ClampMin = "0", UIMax = "5", EditCondition = "ClusterSizeMethod == EClusterSizeMethod::ByGrid", EditConditionHides))
+	int DriftIterations = 0;
 
 	/** If true, bones will only be added to the same cluster if they are physically connected (either directly, or via other bones in the same cluster) */
 	UPROPERTY(EditAnywhere, Category = AutoCluster, meta = (DisplayName = "Enforce Cluster Connectivity"))
@@ -118,9 +118,24 @@ public:
 	virtual void Render(const FSceneView* View, FViewport* Viewport, FPrimitiveDrawInterface* PDI) override;
 	void DrawBox(FPrimitiveDrawInterface* PDI, FVector Center, float SideLength);
 
+	virtual void FractureContextChanged() override;
+
 	virtual void Execute(TWeakPtr<FFractureEditorModeToolkit> InToolkit) override;
 
 	UPROPERTY(EditAnywhere, Category = AutoCluster)
 	TObjectPtr<UFractureAutoClusterSettings> AutoClusterSettings;
+
+protected:
+	virtual void ClearVisualizations() override
+	{
+		Super::ClearVisualizations();
+		ShowGridPoints.Empty();
+		GridPointsMappings.Empty();
+	}
+
+private:
+
+	TArray<FVector> ShowGridPoints;
+	FVisualizationMappings GridPointsMappings;
 };
 
diff --git a/Engine/Plugins/Experimental/Fracture/Source/FractureEngine/Private/FractureEngineClustering.cpp b/Engine/Plugins/Experimental/Fracture/Source/FractureEngine/Private/FractureEngineClustering.cpp
--- a/Engine/Plugins/Experimental/Fracture/Source/FractureEngine/Private/FractureEngineClustering.cpp
+++ b/Engine/Plugins/Experimental/Fracture/Source/FractureEngine/Private/FractureEngineClustering.cpp
@@ -24,7 +24,8 @@ void FVoronoiPartitioner::KMeansPartition(int32 InPartitionCount, int32 MaxItera
 	}
 
 	InitializePartitions(InitialCenters);
-	for (int32 Iteration = 0; Iteration < MaxIterations; Iteration++)
+	// Note: We run max iterations + 1 here because the first iteration just applies the initial cluster assignment
+	for (int32 Iteration = 0; Iteration < MaxIterations + 1; Iteration++)
 	{
 		// Refinement is complete when no nodes change partition.
 		if (!Refine())
@@ -593,28 +594,7 @@ void FFractureEngineClustering::AutoCluster(FGeometryCollection& GeometryCollect
 	}
 	else if (ClusterSizeMethod == EFractureEngineClusterSizeMethod::ByGrid)
 	{
-		FBox Bounds = FVoronoiPartitioner::GenerateBounds(&GeometryCollection, ClusterIndex);
-		int32 GridX = FMath::Max(1, InGridX);
-		int32 GridY = FMath::Max(1, InGridY);
-		int32 GridZ = FMath::Max(1, InGridZ);
-		PartitionPositions.Reserve(GridX * GridY * GridZ);
-
-		for (int32 StepX = 0; StepX < GridX; ++StepX)
-		{
-			double Xt = GridX == 1 ? .5f : double(StepX) / double(GridX - 1);
-			double X = FMath::Lerp(Bounds.Min.X, Bounds.Max.X, Xt);
-			for (int32 StepY = 0; StepY < GridY; ++StepY)
-			{
-				double Yt = GridY == 1 ? .5f : double(StepY) / double(GridY - 1);
-				double Y = FMath::Lerp(Bounds.Min.Y, Bounds.Max.Y, Yt);
-				for (int32 StepZ = 0; StepZ < GridZ; ++StepZ)
-				{
-					double Zt = GridZ == 1 ? .5f : double(StepZ) / double(GridZ - 1);
-					double Z = FMath::Lerp(Bounds.Min.Z, Bounds.Max.Z, Zt);
-					PartitionPositions.Emplace(X, Y, Z);
-				}
-			}
-		}
+		PartitionPositions = GenerateGridSites(GeometryCollection, ClusterIndex, InGridX, InGridY, InGridZ);
 	}
 
 	int32 SiteCountToUse = DesiredSiteCountToUse;
@@ -702,4 +682,62 @@ void FFractureEngineClustering::AutoCluster(FGeometryCollection& GeometryCollect
 	FGeometryCollectionClusteringUtility::ValidateResults(&GeometryCollection);
 }
 
+TArray<FVector> FFractureEngineClustering::GenerateGridSites(
+	const FGeometryCollection& GeometryCollection,
+	const TArray<int32>& BoneIndices,
+	const int32 GridX,
+	const int32 GridY,
+	const int32 GridZ)
+{
+	TArray<FVector> AllPoints;
+	FFractureEngineBoneSelection Selection(GeometryCollection, BoneIndices);
+
+	Selection.ConvertSelectionToClusterNodes();
+
+	for (const int32 ClusterIndex : Selection.GetSelectedBones())
+	{
+		AllPoints.Append(GenerateGridSites(GeometryCollection, ClusterIndex, GridX, GridY, GridZ));
+	}
+	return AllPoints;
+}
+
+TArray<FVector> FFractureEngineClustering::GenerateGridSites(
+	const FGeometryCollection& GeometryCollection,
+	const int32 ClusterIndex,
+	const int32 InGridX,
+	const int32 InGridY,
+	const int32 InGridZ)
+{
+	TArray<FVector> PartitionPositions;
+
+	FBox Bounds = FVoronoiPartitioner::GenerateBounds(&GeometryCollection, ClusterIndex);
+	int32 GridX = FMath::Max(1, InGridX);
+	int32 GridY = FMath::Max(1, InGridY);
+	int32 GridZ = FMath::Max(1, InGridZ);
+	PartitionPositions.Reserve(GridX * GridY * GridZ);
+
+	auto StepToParam = [](int32 Step, int32 NumSteps) -> double
+	{
+		return (double(Step) + .5) / double(NumSteps);
+	};
+
+	for (int32 StepX = 0; StepX < GridX; ++StepX)
+	{
+		double Xt = StepToParam(StepX, GridX);
+		double X = FMath::Lerp(Bounds.Min.X, Bounds.Max.X, Xt);
+		for (int32 StepY = 0; StepY < GridY; ++StepY)
+		{
+			double Yt = StepToParam(StepY, GridY);
+			double Y = FMath::Lerp(Bounds.Min.Y, Bounds.Max.Y, Yt);
+			for (int32 StepZ = 0; StepZ < GridZ; ++StepZ)
+			{
+				double Zt = StepToParam(StepZ, GridZ);
+				double Z = FMath::Lerp(Bounds.Min.Z, Bounds.Max.Z, Zt);
+				PartitionPositions.Emplace(X, Y, Z);
+			}
+		}
+	}
+
+	return PartitionPositions;
+}
 
diff --git a/Engine/Plugins/Experimental/Fracture/Source/FractureEngine/Public/FractureEngineClustering.h b/Engine/Plugins/Experimental/Fracture/Source/FractureEngine/Public/FractureEngineClustering.h
--- a/Engine/Plugins/Experimental/Fracture/Source/FractureEngine/Public/FractureEngineClustering.h
+++ b/Engine/Plugins/Experimental/Fracture/Source/FractureEngine/Public/FractureEngineClustering.h
@@ -115,6 +115,20 @@ public:
 		const int32 GridZ = 2,
 		const float MinimumClusterSize = 0,
 		const int32 KMeansIterations = 500);
+
+	static TArray<FVector> GenerateGridSites(
+		const FGeometryCollection& GeometryCollection,
+		const TArray<int32>& BoneIndices,
+		const int32 GridX,
+		const int32 GridY,
+		const int32 GridZ);
+
+	static TArray<FVector> GenerateGridSites(
+		const FGeometryCollection& GeometryCollection,
+		const int32 ClusterIndex,
+		const int32 GridX,
+		const int32 GridY,
+		const int32 GridZ);
 };
 
 #if UE_ENABLE_INCLUDE_ORDER_DEPRECATED_IN_5_2
