GeometryCollection Dataflow minor robustness fixes:
- Validate bone selection before applying fracture
- Validate that collection is not empty before applying back to asset, warn user if not

diff --git a/Engine/Plugins/Experimental/Fracture/Source/FractureEngine/Private/FractureEngineFracturing.cpp b/Engine/Plugins/Experimental/Fracture/Source/FractureEngine/Private/FractureEngineFracturing.cpp
--- a/Engine/Plugins/Experimental/Fracture/Source/FractureEngine/Private/FractureEngineFracturing.cpp
+++ b/Engine/Plugins/Experimental/Fracture/Source/FractureEngine/Private/FractureEngineFracturing.cpp
@@ -1,6 +1,7 @@
 // Copyright Epic Games, Inc. All Rights Reserved.
 
 #include "FractureEngineFracturing.h"
+#include "FractureEngineSelection.h"
 #include "Dataflow/DataflowSelection.h"
 #include "GeometryCollection/GeometryCollectionClusteringUtility.h"
 #include "GeometryCollection/GeometryCollectionAlgo.h"
@@ -245,6 +246,10 @@ void FFractureEngineFracturing::VoronoiFracture(FManagedArrayCollection& InOutCo
 				VoronoiPlanarCells.InternalSurfaceMaterials.NoiseSettings = NoiseSettings;
 
 				TArray<int32> TransformSelectionArr = InTransformSelection.AsArray();
+				if (!FFractureEngineSelection::IsBoneSelectionValid(InOutCollection, TransformSelectionArr))
+				{
+					return;
+				}
 				
 				int ResultGeometryIndex = CutMultipleWithPlanarCells(VoronoiPlanarCells, *GeomCollection, TransformSelectionArr, InGrout, InCollisionSampleSpacing, InRandomSeed, FTransform().Identity);
 
@@ -312,6 +317,10 @@ void FFractureEngineFracturing::PlaneCutter(FManagedArrayCollection& InOutCollec
 		float GroutVal = InGrout;
 
 		TArray<int32> TransformSelectionArr = InTransformSelection.AsArray();
+		if (!FFractureEngineSelection::IsBoneSelectionValid(InOutCollection, TransformSelectionArr))
+		{
+			return;
+		}
 
 		int ResultGeometryIndex = CutMultipleWithMultiplePlanes(CuttingPlanes, InternalSurfaceMaterials, *GeomCollection, TransformSelectionArr, GroutVal, CollisionSampleSpacingVal, InRandomSeed, FTransform().Identity);
 
diff --git a/Engine/Plugins/Experimental/Fracture/Source/FractureEngine/Private/FractureEngineSelection.cpp b/Engine/Plugins/Experimental/Fracture/Source/FractureEngine/Private/FractureEngineSelection.cpp
--- a/Engine/Plugins/Experimental/Fracture/Source/FractureEngine/Private/FractureEngineSelection.cpp
+++ b/Engine/Plugins/Experimental/Fracture/Source/FractureEngine/Private/FractureEngineSelection.cpp
@@ -9,7 +9,7 @@
 
 static void ConvertSelectionSetToSelectionArr(const TSet<int32>& SelectionSet, TArray<int32>& SelectionArr)
 {
-	SelectionArr.Empty();
+	SelectionArr.Empty(SelectionSet.Num());
 
 	for (auto& BoneIdx : SelectionSet)
 	{
@@ -17,6 +17,25 @@ static void ConvertSelectionSetToSelectionArr(const TSet<int32>& SelectionSet, T
 	}
 }
 
+
+bool FFractureEngineSelection::IsBoneSelectionValid(const FManagedArrayCollection& Collection, const TArray<int32>& SelectedBones)
+{
+	return IsSelectionValid(Collection, SelectedBones, FGeometryCollection::TransformGroup);
+}
+
+bool FFractureEngineSelection::IsSelectionValid(const FManagedArrayCollection& Collection, const TArray<int32>& SelectedItems, const FName ItemGroup)
+{
+	int32 NumItems = Collection.NumElements(ItemGroup);
+	for (int32 Idx : SelectedItems)
+	{
+		if (Idx < 0 || Idx >= NumItems)
+		{
+			return false;
+		}
+	}
+	return true;
+}
+
 void FFractureEngineSelection::GetRootBones(const FManagedArrayCollection& Collection, TArray<int32>& RootBonesOut)
 {
 	if (Collection.HasGroup(FGeometryCollection::TransformGroup) &&
diff --git a/Engine/Plugins/Experimental/Fracture/Source/FractureEngine/Public/FractureEngineSelection.h b/Engine/Plugins/Experimental/Fracture/Source/FractureEngine/Public/FractureEngineSelection.h
--- a/Engine/Plugins/Experimental/Fracture/Source/FractureEngine/Public/FractureEngineSelection.h
+++ b/Engine/Plugins/Experimental/Fracture/Source/FractureEngine/Public/FractureEngineSelection.h
@@ -7,6 +7,7 @@
 class FGeometryCollection;
 struct FDataflowTransformSelection;
 struct FManagedArrayCollection;
+class FName;
 
 class FRACTUREENGINE_API FFractureEngineSelection
 {
@@ -43,6 +44,9 @@ public:
 
 	static void SelectByVolume(FGeometryCollection& GeometryCollection, TArray<int32>& SelectedBones, const float VolumeMin, const float VolumeMax);
 	static void SelectByVolume(FGeometryCollection& GeometryCollection, FDataflowTransformSelection& TransformSelection, const float VolumeMin, const float VolumeMax);
+
+	static bool IsBoneSelectionValid(const FManagedArrayCollection& Collection, const TArray<int32>& SelectedBones);
+	static bool IsSelectionValid(const FManagedArrayCollection& Collection, const TArray<int32>& SelectedItems, const FName ItemGroup);
 };
 
 #if UE_ENABLE_INCLUDE_ORDER_DEPRECATED_IN_5_2
diff --git a/Engine/Plugins/Experimental/GeometryCollectionPlugin/Source/GeometryCollectionNodes/Private/Dataflow/Nodes/GeometryCollectionAssetNodes.cpp b/Engine/Plugins/Experimental/GeometryCollectionPlugin/Source/GeometryCollectionNodes/Private/Dataflow/Nodes/GeometryCollectionAssetNodes.cpp
--- a/Engine/Plugins/Experimental/GeometryCollectionPlugin/Source/GeometryCollectionNodes/Private/Dataflow/Nodes/GeometryCollectionAssetNodes.cpp
+++ b/Engine/Plugins/Experimental/GeometryCollectionPlugin/Source/GeometryCollectionNodes/Private/Dataflow/Nodes/GeometryCollectionAssetNodes.cpp
@@ -44,6 +44,12 @@ void FGeometryCollectionTerminalDataflowNode::SetAssetValue(TObjectPtr<UObject>
 			const FMaterialArray& InMaterials = GetValue<FMaterialArray>(Context, &Materials);
 			const FInstancedMeshesArray& InInstancedMeshes = GetValue<FInstancedMeshesArray>(Context, &InstancedMeshes);
 
+			if (InCollection.NumElements(FGeometryCollection::TransformGroup) == 0)
+			{
+				UE_LOG(LogChaos, Warning, TEXT("Cannot set Geometry Collection asset: Geometry Collection must have at least one transform."));
+				return;
+			}
+
 			const bool bHasInternalMaterial = false; // with data flow there's no assumption of internal materials
 			CollectionAsset->ResetFrom(InCollection, InMaterials, false);
 			CollectionAsset->AutoInstanceMeshes = InInstancedMeshes;
