Stop auto-cluster from fixing isolated clusters via merging, because it creates undesired (too large) clusters.

diff --git a/Engine/Plugins/Experimental/ChaosEditor/Source/FractureEditor/Private/FractureToolAutoCluster.h b/Engine/Plugins/Experimental/ChaosEditor/Source/FractureEditor/Private/FractureToolAutoCluster.h
--- a/Engine/Plugins/Experimental/ChaosEditor/Source/FractureEditor/Private/FractureToolAutoCluster.h
+++ b/Engine/Plugins/Experimental/ChaosEditor/Source/FractureEditor/Private/FractureToolAutoCluster.h
@@ -98,7 +98,7 @@ public:
 	UPROPERTY(EditAnywhere, Category = AutoCluster, meta = (EditCondition = "bEnforceConnectivity == true && ClusterSizeMethod != EClusterSizeMethod::ByGrid"))
 	bool bEnforceSiteParameters = true;
 
-	/** If true, prevent the creation of clusters with only a single child. Either by merging into a neighboring cluster, or not creating the cluster. */
+	/** If true, prevent the creation of clusters with only a single child, skipping creation of a new cluster in such cases. */
 	UPROPERTY(EditAnywhere, Category = AutoCluster)
 	bool bAvoidIsolated = true;
 };
diff --git a/Engine/Plugins/Experimental/Fracture/Source/FractureEngine/Private/FractureEngineClustering.cpp b/Engine/Plugins/Experimental/Fracture/Source/FractureEngine/Private/FractureEngineClustering.cpp
--- a/Engine/Plugins/Experimental/Fracture/Source/FractureEngine/Private/FractureEngineClustering.cpp
+++ b/Engine/Plugins/Experimental/Fracture/Source/FractureEngine/Private/FractureEngineClustering.cpp
@@ -620,11 +620,10 @@ void FFractureEngineClustering::AutoCluster(FGeometryCollection& GeometryCollect
 			VoronoiPartition.SplitDisconnectedPartitions(&GeometryCollection);
 		}
 
-		if (bAvoidIsolated)
-		{
-			// attempt to remove isolated via merging (may not succeed, as it only merges if there is a cluster in proximity)
-			VoronoiPartition.MergeSingleElementPartitions(&GeometryCollection);
-		}
+		// Note: Previously if bAvoidIsolated was set, we'd attempt to merge clusters of one here via
+		//  VoronoiPartition.MergeSingleElementPartitions(&GeometryCollection);
+		// However this results in some overly-large clusters, especially with grid clustering,
+		// so we prefer to simply skip clustering in such cases.
 
 		if (MinimumClusterSize > 0)
 		{
