add more dataflow nodes: Cluster Selected, Merge Clusters, and Select Target Level
+ move the common code to the Fracture Engine module
+ refactor some of the utility functions to work on Managed Array Collection

diff --git a/Engine/Plugins/Experimental/ChaosEditor/Source/FractureEditor/Private/FractureToolClustering.cpp b/Engine/Plugins/Experimental/ChaosEditor/Source/FractureEditor/Private/FractureToolClustering.cpp
--- a/Engine/Plugins/Experimental/ChaosEditor/Source/FractureEditor/Private/FractureToolClustering.cpp
+++ b/Engine/Plugins/Experimental/ChaosEditor/Source/FractureEditor/Private/FractureToolClustering.cpp
@@ -7,6 +7,7 @@
 #include "GeometryCollection/GeometryCollectionComponent.h"
 #include "GeometryCollection/GeometryCollection.h"
 #include "GeometryCollection/GeometryCollectionClusteringUtility.h"
+#include "FractureEngineClustering.h"
 #include "ScopedTransaction.h"
 
 #include UE_INLINE_GENERATED_CPP_BY_NAME(FractureToolClustering)
@@ -121,23 +122,8 @@ void UFractureToolCluster::Execute(TWeakPtr<FFractureEditorModeToolkit> InToolki
 		{
 			FGeometryCollectionEdit Edit(Context.GetGeometryCollectionComponent(), GeometryCollection::EEditUpdate::RestPhysicsDynamic);
 
-			int32 StartTransformCount = Context.GetGeometryCollection()->Transform.Num();
-			
-			Context.RemoveRootNodes();
-			Context.Sanitize();
-
-			if (Context.GetSelection().Num() > 1)
+			if (FFractureEngineClustering::ClusterSelected(*Context.GetGeometryCollection(), Context.GetSelection()))
 			{
-				// Cluster selected bones beneath common parent
-				int32 LowestCommonAncestor = FGeometryCollectionClusteringUtility::FindLowestCommonAncestor(Context.GetGeometryCollection().Get(), Context.GetSelection());
-
-				if (LowestCommonAncestor != INDEX_NONE)
-				{
-					FGeometryCollectionClusteringUtility::ClusterBonesUnderNewNodeWithParent(Context.GetGeometryCollection().Get(), LowestCommonAncestor, Context.GetSelection(), true);
-				}
-
-				Context.GenerateGuids(StartTransformCount);
-
 				Refresh(Context, Toolkit);
 			}
 		}
@@ -293,52 +279,9 @@ void UFractureToolClusterMerge::Execute(TWeakPtr<FFractureEditorModeToolkit> InT
 		{
 			FGeometryCollectionEdit Edit(Context.GetGeometryCollectionComponent(), GeometryCollection::EEditUpdate::RestPhysicsDynamic);
 
-			const TManagedArray<TSet<int32>>& Children = Context.GetGeometryCollection()->Children;
-
-			Context.ConvertEmbeddedSelectionToParents(); // embedded geo must stay attached to parent
-
-			// Collect children of context clusters
-			TArray<int32> ChildBones;
-			ChildBones.Reserve(Context.GetGeometryCollection()->NumElements(FGeometryCollection::TransformGroup));
-			bool bHasClusters = false;
-			for (int32 Select : Context.GetSelection())
-			{
-				if (Context.GetGeometryCollection()->SimulationType[Select] == FGeometryCollection::ESimulationTypes::FST_Clustered)
-				{
-					bHasClusters = true;
-					ChildBones.Append(Children[Select].Array());
-				}
-				else
-				{
-					ChildBones.Add(Select);
-				}
-			}
-			
-			// If there were no clusters in the selection, we create a cluster -- i.e., follow the logic of UFractureToolCluster::Execute
-			if (!bHasClusters)
+			if (FFractureEngineClustering::MergeSelectedClusters(*Context.GetGeometryCollection(), Context.GetSelection()))
 			{
-				// Cluster selected bones beneath common parent
-				int32 LowestCommonAncestor = FGeometryCollectionClusteringUtility::FindLowestCommonAncestor(Context.GetGeometryCollection().Get(), Context.GetSelection());
-				if (LowestCommonAncestor != INDEX_NONE)
-				{
-					int32 StartTransformCount = Context.GetGeometryCollection()->Transform.Num();
-					FGeometryCollectionClusteringUtility::ClusterBonesUnderNewNodeWithParent(Context.GetGeometryCollection().Get(), LowestCommonAncestor, Context.GetSelection(), true);
-					Context.GenerateGuids(StartTransformCount);
-					Context.SetSelection({ LowestCommonAncestor });
-					Refresh(Context, Toolkit);
-				}
-			}
-			else
-			{
-				int32 MergeNode = FGeometryCollectionClusteringUtility::PickBestNodeToMergeTo(Context.GetGeometryCollection().Get(), Context.GetSelection());
-				if (MergeNode >= 0)
-				{
-					FGeometryCollectionClusteringUtility::ClusterBonesUnderExistingNode(Context.GetGeometryCollection().Get(), MergeNode, ChildBones);
-					FGeometryCollectionClusteringUtility::RemoveDanglingClusters(Context.GetGeometryCollection().Get());
-
-					Context.SetSelection({ MergeNode });
-					Refresh(Context, Toolkit);
-				}
+				Refresh(Context, Toolkit);
 			}
 		}
 
diff --git a/Engine/Plugins/Experimental/Fracture/Source/FractureEngine/Private/FractureEngineClustering.cpp b/Engine/Plugins/Experimental/Fracture/Source/FractureEngine/Private/FractureEngineClustering.cpp
--- a/Engine/Plugins/Experimental/Fracture/Source/FractureEngine/Private/FractureEngineClustering.cpp
+++ b/Engine/Plugins/Experimental/Fracture/Source/FractureEngine/Private/FractureEngineClustering.cpp
@@ -2,6 +2,10 @@
 
 #include "FractureEngineClustering.h"
 #include "GeometryCollection/GeometryCollection.h"
+#include "GeometryCollection/GeometryCollectionUtility.h"
+#include "GeometryCollection/Facades/CollectionTransformSelectionFacade.h"
+#include "GeometryCollection/Facades/CollectionTransformFacade.h"
+#include "GeometryCollection/Facades/CollectionHierarchyFacade.h"
 #include "GeometryCollection/GeometryCollectionClusteringUtility.h"
 #include "GeometryCollection/GeometryCollectionAlgo.h"
 #include "Math/BoxSphereBounds.h"
@@ -745,3 +749,100 @@ TArray<FVector> FFractureEngineClustering::GenerateGridSites(
 	return PartitionPositions;
 }
 
+
+bool FFractureEngineClustering::ClusterSelected(
+	FGeometryCollection& GeometryCollection,
+	TArray<int32>& Selection)
+{
+	GeometryCollection::Facades::FCollectionTransformSelectionFacade SelectionFacade(GeometryCollection);
+	SelectionFacade.RemoveRootNodes(Selection);
+	SelectionFacade.Sanitize(Selection);
+
+	if (Selection.Num() <= 1)
+	{
+		// Don't make a cluster of 1
+		return false;
+	}
+
+	Chaos::Facades::FCollectionHierarchyFacade HierarchyFacade(GeometryCollection);
+	int32 StartTransformCount = GeometryCollection.NumElements(FGeometryCollection::TransformGroup);
+	int32 LowestCommonAncestor = FGeometryCollectionClusteringUtility::FindLowestCommonAncestor(&GeometryCollection, Selection);
+	// Cluster selected bones beneath common parent
+	if (LowestCommonAncestor != INDEX_NONE)
+	{
+		FGeometryCollectionClusteringUtility::ClusterBonesUnderNewNodeWithParent(&GeometryCollection, LowestCommonAncestor, Selection, true);
+		GeometryCollection::GenerateTemporaryGuids(&GeometryCollection, StartTransformCount, true);
+		FGeometryCollectionClusteringUtility::RemoveDanglingClusters(&GeometryCollection);
+		HierarchyFacade.GenerateLevelAttribute();
+
+		return true;
+	}
+
+	return false;
+}
+
+bool FFractureEngineClustering::MergeSelectedClusters(FGeometryCollection& GeometryCollection, TArray<int32>& Selection)
+{
+	Chaos::Facades::FCollectionHierarchyFacade HierarchyFacade(GeometryCollection);
+
+	GeometryCollection::Facades::FCollectionTransformSelectionFacade SelectionFacade(GeometryCollection);
+	SelectionFacade.ConvertEmbeddedSelectionToParents(Selection); // embedded geo must stay attached to parent
+
+	// Collect children of context clusters
+	TArray<int32> ChildBones;
+	int32 StartTransformCount = GeometryCollection.NumElements(FGeometryCollection::TransformGroup);
+	ChildBones.Reserve(StartTransformCount);
+	bool bHasClusters = false;
+	for (int32 Select : Selection)
+	{
+		if (GeometryCollection.SimulationType[Select] == FGeometryCollection::ESimulationTypes::FST_Clustered)
+		{
+			bHasClusters = true;
+			ChildBones.Append(HierarchyFacade.GetChildrenAsArray(Select));
+		}
+		else
+		{
+			ChildBones.Add(Select);
+		}
+	}
+
+	// If there were no clusters in the selection, we create a cluster
+	if (!bHasClusters)
+	{
+		// Cluster selected bones beneath common parent
+		int32 LowestCommonAncestor = FGeometryCollectionClusteringUtility::FindLowestCommonAncestor(&GeometryCollection, Selection);
+		if (LowestCommonAncestor != INDEX_NONE)
+		{
+			FGeometryCollectionClusteringUtility::ClusterBonesUnderNewNodeWithParent(&GeometryCollection, LowestCommonAncestor, Selection, true);
+			GeometryCollection::GenerateTemporaryGuids(&GeometryCollection, StartTransformCount, true);
+			Selection.SetNum(1);
+			Selection[0] = LowestCommonAncestor;
+			if (FGeometryCollectionClusteringUtility::RemoveDanglingClusters(&GeometryCollection))
+			{
+				// Bone index is unreliable after removal of some bones
+				Selection.Empty();
+			}
+			HierarchyFacade.GenerateLevelAttribute();
+			return true;
+		}
+	}
+	else
+	{
+		int32 MergeNode = FGeometryCollectionClusteringUtility::PickBestNodeToMergeTo(&GeometryCollection, Selection);
+		if (MergeNode >= 0)
+		{
+			FGeometryCollectionClusteringUtility::ClusterBonesUnderExistingNode(&GeometryCollection, MergeNode, ChildBones);
+			Selection.SetNum(1);
+			Selection[0] = MergeNode;
+			if (FGeometryCollectionClusteringUtility::RemoveDanglingClusters(&GeometryCollection))
+			{
+				// Bone index is unreliable after removal of some bones
+				Selection.Empty();
+			}
+			HierarchyFacade.GenerateLevelAttribute();
+			return true;
+		}
+	}
+
+	return false;
+}
diff --git a/Engine/Plugins/Experimental/Fracture/Source/FractureEngine/Public/FractureEngineClustering.h b/Engine/Plugins/Experimental/Fracture/Source/FractureEngine/Public/FractureEngineClustering.h
--- a/Engine/Plugins/Experimental/Fracture/Source/FractureEngine/Public/FractureEngineClustering.h
+++ b/Engine/Plugins/Experimental/Fracture/Source/FractureEngine/Public/FractureEngineClustering.h
@@ -130,6 +130,20 @@ public:
 		const int32 GridY,
 		const int32 GridZ,
 		FBox* OutBounds = nullptr);
+
+	// Cluster the chosen transform indices (and update the selection array to remove any that were not clustered, i.e. invalid or root transforms)
+	// @return true if the GeometryCollection was updated
+	static bool ClusterSelected(
+		FGeometryCollection& GeometryCollection,
+		TArray<int32>& InOutSelection
+	);
+
+	// Merge selected clusters. Non-clusters in the selection are converted to the closest (parent) clusters.
+	// On success, returns true and InOutSelection holds the index of the cluster to which the selection was merged.
+	static bool MergeSelectedClusters(
+		FGeometryCollection& GeometryCollection,
+		TArray<int32>& InOutSelection
+	);
 };
 
 #if UE_ENABLE_INCLUDE_ORDER_DEPRECATED_IN_5_2
diff --git a/Engine/Plugins/Experimental/GeometryCollectionPlugin/Source/GeometryCollectionNodes/Private/Dataflow/GeometryCollectionClusteringNodes.cpp b/Engine/Plugins/Experimental/GeometryCollectionPlugin/Source/GeometryCollectionNodes/Private/Dataflow/GeometryCollectionClusteringNodes.cpp
--- a/Engine/Plugins/Experimental/GeometryCollectionPlugin/Source/GeometryCollectionNodes/Private/Dataflow/GeometryCollectionClusteringNodes.cpp
+++ b/Engine/Plugins/Experimental/GeometryCollectionPlugin/Source/GeometryCollectionNodes/Private/Dataflow/GeometryCollectionClusteringNodes.cpp
@@ -42,6 +42,8 @@ namespace Dataflow
 		DATAFLOW_NODE_REGISTER_CREATION_FACTORY(FAutoClusterDataflowNode);
 		DATAFLOW_NODE_REGISTER_CREATION_FACTORY(FClusterFlattenDataflowNode);
 		DATAFLOW_NODE_REGISTER_CREATION_FACTORY(FClusterUnclusterDataflowNode);
+		DATAFLOW_NODE_REGISTER_CREATION_FACTORY(FClusterDataflowNode);
+		DATAFLOW_NODE_REGISTER_CREATION_FACTORY(FClusterMergeDataflowNode);
 
 		// GeometryCollection|Cluster
 		DATAFLOW_NODE_REGISTER_CREATION_FACTORY_NODE_COLORS_BY_CATEGORY("GeometryCollection|Cluster", FLinearColor(.25f, 0.45f, 0.8f), CDefaultNodeBodyTintColor);
@@ -147,3 +149,35 @@ void FClusterUnclusterDataflowNode::Evaluate(Dataflow::FContext& Context, const
 		}
 	}
 }
+
+
+void FClusterDataflowNode::Evaluate(Dataflow::FContext& Context, const FDataflowOutput* Out) const
+{
+	if (Out->IsA(&Collection))
+	{
+		const FManagedArrayCollection& InCollection = GetValue(Context, &Collection);
+		const FDataflowTransformSelection& InTransformSelection = GetValue(Context, &TransformSelection);
+		if (TUniquePtr<FGeometryCollection> GeomCollection = TUniquePtr<FGeometryCollection>(InCollection.NewCopy<FGeometryCollection>()))
+		{
+			TArray<int32> Selection = InTransformSelection.AsArray();
+			FFractureEngineClustering::ClusterSelected(*GeomCollection, Selection);
+			SetValue(Context, (const FManagedArrayCollection&)(*GeomCollection), &Collection);
+		}
+	}
+}
+
+void FClusterMergeDataflowNode::Evaluate(Dataflow::FContext& Context, const FDataflowOutput* Out) const
+{
+	if (Out->IsA(&Collection))
+	{
+		const FManagedArrayCollection& InCollection = GetValue(Context, &Collection);
+		const FDataflowTransformSelection& InTransformSelection = GetValue(Context, &TransformSelection);
+		if (TUniquePtr<FGeometryCollection> GeomCollection = TUniquePtr<FGeometryCollection>(InCollection.NewCopy<FGeometryCollection>()))
+		{
+			TArray<int32> Selection = InTransformSelection.AsArray();
+			FFractureEngineClustering::MergeSelectedClusters(*GeomCollection, Selection);
+			SetValue(Context, (const FManagedArrayCollection&)(*GeomCollection), &Collection);
+		}
+	}
+}
+
diff --git a/Engine/Plugins/Experimental/GeometryCollectionPlugin/Source/GeometryCollectionNodes/Private/Dataflow/GeometryCollectionSelectionNodes.cpp b/Engine/Plugins/Experimental/GeometryCollectionPlugin/Source/GeometryCollectionNodes/Private/Dataflow/GeometryCollectionSelectionNodes.cpp
--- a/Engine/Plugins/Experimental/GeometryCollectionPlugin/Source/GeometryCollectionNodes/Private/Dataflow/GeometryCollectionSelectionNodes.cpp
+++ b/Engine/Plugins/Experimental/GeometryCollectionPlugin/Source/GeometryCollectionNodes/Private/Dataflow/GeometryCollectionSelectionNodes.cpp
@@ -51,6 +51,7 @@ namespace Dataflow
 		DATAFLOW_NODE_REGISTER_CREATION_FACTORY(FCollectionTransformSelectionChildrenDataflowNode);
 		DATAFLOW_NODE_REGISTER_CREATION_FACTORY(FCollectionTransformSelectionSiblingsDataflowNode);
 		DATAFLOW_NODE_REGISTER_CREATION_FACTORY(FCollectionTransformSelectionLevelDataflowNode);
+		DATAFLOW_NODE_REGISTER_CREATION_FACTORY(FCollectionTransformSelectionTargetLevelDataflowNode);
 		DATAFLOW_NODE_REGISTER_CREATION_FACTORY(FCollectionTransformSelectionContactDataflowNode);
 		DATAFLOW_NODE_REGISTER_CREATION_FACTORY(FCollectionTransformSelectionLeafDataflowNode);
 		DATAFLOW_NODE_REGISTER_CREATION_FACTORY(FCollectionTransformSelectionClusterDataflowNode);
@@ -517,6 +518,31 @@ void FCollectionTransformSelectionLevelDataflowNode::Evaluate(Dataflow::FContext
 }
 
 
+void FCollectionTransformSelectionTargetLevelDataflowNode::Evaluate(Dataflow::FContext& Context, const FDataflowOutput* Out) const
+{
+	if (Out->IsA(&TransformSelection))
+	{
+		const FManagedArrayCollection& InCollection = GetValue(Context, &Collection);
+		GeometryCollection::Facades::FCollectionTransformSelectionFacade TransformSelectionFacade(InCollection);
+
+		int InTargetLevel = GetValue(Context, &TargetLevel);
+
+		TArray<int32> AllAtLevel = TransformSelectionFacade.GetBonesExactlyAtLevel(InTargetLevel, bSkipEmbedded);
+
+		FDataflowTransformSelection NewTransformSelection;
+		NewTransformSelection.Initialize(InCollection.NumElements(FGeometryCollection::TransformGroup), false);
+		NewTransformSelection.SetFromArray(AllAtLevel);
+
+		SetValue(Context, MoveTemp(NewTransformSelection), &TransformSelection);
+	}
+	else if (Out->IsA(&Collection))
+	{
+		const FManagedArrayCollection& InCollection = GetValue(Context, &Collection);
+		SetValue(Context, InCollection, &Collection);
+	}
+}
+
+
 void FCollectionTransformSelectionContactDataflowNode::Evaluate(Dataflow::FContext& Context, const FDataflowOutput* Out) const
 {
 	if (Out->IsA<FDataflowTransformSelection>(&TransformSelection))
diff --git a/Engine/Plugins/Experimental/GeometryCollectionPlugin/Source/GeometryCollectionNodes/Public/Dataflow/GeometryCollectionClusteringNodes.h b/Engine/Plugins/Experimental/GeometryCollectionPlugin/Source/GeometryCollectionNodes/Public/Dataflow/GeometryCollectionClusteringNodes.h
--- a/Engine/Plugins/Experimental/GeometryCollectionPlugin/Source/GeometryCollectionNodes/Public/Dataflow/GeometryCollectionClusteringNodes.h
+++ b/Engine/Plugins/Experimental/GeometryCollectionPlugin/Source/GeometryCollectionNodes/Public/Dataflow/GeometryCollectionClusteringNodes.h
@@ -173,6 +173,70 @@ public:
 
 };
 
+/**
+ * Cluster selected nodes under a new parent
+ */
+USTRUCT(meta = (DataflowGeometryCollection))
+struct FClusterDataflowNode : public FDataflowNode
+{
+	GENERATED_USTRUCT_BODY()
+	DATAFLOW_NODE_DEFINE_INTERNAL(FClusterDataflowNode, "Cluster", "GeometryCollection|Cluster", "")
+
+public:
+
+	/** Collection on which to cluster nodes */
+	UPROPERTY(meta = (DataflowInput, DataflowOutput, DataflowPassthrough = "Collection", DataflowIntrinsic))
+	FManagedArrayCollection Collection;
+
+	/** Bone selection */
+	UPROPERTY(meta = (DataflowInput, DisplayName = "TransformSelection"))
+	FDataflowTransformSelection TransformSelection;
+
+	FClusterDataflowNode(const Dataflow::FNodeParameters& InParam, FGuid InGuid = FGuid::NewGuid())
+		: FDataflowNode(InParam, InGuid)
+	{
+		RegisterInputConnection(&Collection);
+		RegisterInputConnection(&TransformSelection);
+		RegisterOutputConnection(&Collection, &Collection);
+	}
+
+	virtual void Evaluate(Dataflow::FContext& Context, const FDataflowOutput* Out) const override;
+
+};
+
+
+/**
+ * Merge selected clusters under a new parent
+ */
+USTRUCT(meta = (DataflowGeometryCollection))
+struct FClusterMergeDataflowNode : public FDataflowNode
+{
+	GENERATED_USTRUCT_BODY()
+	DATAFLOW_NODE_DEFINE_INTERNAL(FClusterMergeDataflowNode, "Merge", "GeometryCollection|Cluster", "")
+
+public:
+
+	/** Collection on which to merge clusters */
+	UPROPERTY(meta = (DataflowInput, DataflowOutput, DataflowPassthrough = "Collection", DataflowIntrinsic))
+	FManagedArrayCollection Collection;
+
+	/** Bone selection */
+	UPROPERTY(meta = (DataflowInput, DisplayName = "TransformSelection"))
+	FDataflowTransformSelection TransformSelection;
+
+	FClusterMergeDataflowNode(const Dataflow::FNodeParameters& InParam, FGuid InGuid = FGuid::NewGuid())
+		: FDataflowNode(InParam, InGuid)
+	{
+		RegisterInputConnection(&Collection);
+		RegisterInputConnection(&TransformSelection);
+		RegisterOutputConnection(&Collection, &Collection);
+	}
+
+	virtual void Evaluate(Dataflow::FContext& Context, const FDataflowOutput* Out) const override;
+
+};
+
+
 
 namespace Dataflow
 {
diff --git a/Engine/Plugins/Experimental/GeometryCollectionPlugin/Source/GeometryCollectionNodes/Public/Dataflow/GeometryCollectionSelectionNodes.h b/Engine/Plugins/Experimental/GeometryCollectionPlugin/Source/GeometryCollectionNodes/Public/Dataflow/GeometryCollectionSelectionNodes.h
--- a/Engine/Plugins/Experimental/GeometryCollectionPlugin/Source/GeometryCollectionNodes/Public/Dataflow/GeometryCollectionSelectionNodes.h
+++ b/Engine/Plugins/Experimental/GeometryCollectionPlugin/Source/GeometryCollectionNodes/Public/Dataflow/GeometryCollectionSelectionNodes.h
@@ -467,17 +467,17 @@ public:
 
 /**
  *
- * Selects the level of the selected bones
+ * Expand the selection to include all nodes with the same level as the selected nodes
  *
  */
 USTRUCT(meta = (DataflowGeometryCollection))
 struct FCollectionTransformSelectionLevelDataflowNode : public FDataflowNode
 {
 	GENERATED_USTRUCT_BODY()
-	DATAFLOW_NODE_DEFINE_INTERNAL(FCollectionTransformSelectionLevelDataflowNode, "CollectionTransformSelectLevel", "GeometryCollection|Selection|Transform", "")
+	DATAFLOW_NODE_DEFINE_INTERNAL(FCollectionTransformSelectionLevelDataflowNode, "CollectionTransformSelectLevels", "GeometryCollection|Selection|Transform", "")
 
 public:
-	/** Array of the selected bone indicies */
+	/** Array of the selected bone indices */
 	UPROPERTY(meta = (DataflowInput, DataflowOutput, DisplayName = "TransformSelection", DataflowPassthrough = "TransformSelection", DataflowIntrinsic))
 	FDataflowTransformSelection TransformSelection;
 
@@ -499,6 +499,49 @@ public:
 };
 
 
+/**
+ *
+ * Selects the root bones in the Collection
+ *
+ */
+USTRUCT(meta = (DataflowGeometryCollection))
+struct FCollectionTransformSelectionTargetLevelDataflowNode : public FDataflowNode
+{
+	GENERATED_USTRUCT_BODY()
+	DATAFLOW_NODE_DEFINE_INTERNAL(FCollectionTransformSelectionTargetLevelDataflowNode, "CollectionTransformSelectTargetLevel", "GeometryCollection|Selection|Transform", "")
+
+public:
+	/** GeometryCollection for the selection */
+	UPROPERTY(meta = (DataflowInput, DataflowOutput, DataflowPassthrough = "Collection", DataflowIntrinsic))
+	FManagedArrayCollection Collection;
+
+	/** Level to select */
+	UPROPERTY(EditAnywhere, Category = Options, meta = (DataflowInput, ClampMin = 0))
+	int32 TargetLevel = 1;
+
+	/** Whether to avoid embedded geometry in the selection (i.e., only select rigid and cluster nodes) */
+	UPROPERTY(EditAnywhere, Category = Options)
+	bool bSkipEmbedded = false;
+
+	/** Array of the selected bone indices */
+	UPROPERTY(meta = (DataflowOutput, DisplayName = "TransformSelection"))
+	FDataflowTransformSelection TransformSelection;
+
+	FCollectionTransformSelectionTargetLevelDataflowNode(const Dataflow::FNodeParameters& InParam, FGuid InGuid = FGuid::NewGuid())
+		: FDataflowNode(InParam, InGuid)
+	{
+		RegisterInputConnection(&Collection);
+		RegisterInputConnection(&TargetLevel);
+		RegisterOutputConnection(&TransformSelection);
+		RegisterOutputConnection(&Collection, &Collection);
+	}
+
+	virtual void Evaluate(Dataflow::FContext& Context, const FDataflowOutput* Out) const override;
+
+};
+
+
+
 /**
  *
  * Selects the contact(s) of the selected bones
diff --git a/Engine/Source/Runtime/Experimental/Chaos/Private/GeometryCollection/Facades/CollectionHierarchyFacade.cpp b/Engine/Source/Runtime/Experimental/Chaos/Private/GeometryCollection/Facades/CollectionHierarchyFacade.cpp
--- a/Engine/Source/Runtime/Experimental/Chaos/Private/GeometryCollection/Facades/CollectionHierarchyFacade.cpp
+++ b/Engine/Source/Runtime/Experimental/Chaos/Private/GeometryCollection/Facades/CollectionHierarchyFacade.cpp
@@ -65,6 +65,15 @@ namespace Chaos::Facades
 		return ChildrenArray;
 	}
 
+	const TSet<int32>* FCollectionHierarchyFacade::FindChildren(int32 TransformIndex) const
+	{
+		if (ChildrenAttribute.IsValid() && ChildrenAttribute.IsValidIndex(TransformIndex))
+		{
+			return &ChildrenAttribute[TransformIndex];
+		}
+		return nullptr;
+	}
+
 	int32 FCollectionHierarchyFacade::GetParent(int32 TransformIndex) const
 	{
 		return ParentAttribute[TransformIndex];
diff --git a/Engine/Source/Runtime/Experimental/Chaos/Private/GeometryCollection/Facades/CollectionTransformSelectionFacade.cpp b/Engine/Source/Runtime/Experimental/Chaos/Private/GeometryCollection/Facades/CollectionTransformSelectionFacade.cpp
--- a/Engine/Source/Runtime/Experimental/Chaos/Private/GeometryCollection/Facades/CollectionTransformSelectionFacade.cpp
+++ b/Engine/Source/Runtime/Experimental/Chaos/Private/GeometryCollection/Facades/CollectionTransformSelectionFacade.cpp
@@ -132,6 +132,45 @@ namespace GeometryCollection::Facades
 		return OutBones;
 	}
 
+	TArray<int32> FCollectionTransformSelectionFacade::GetBonesExactlyAtLevel(const int32 Level, bool bOnlyClusteredOrRigid) const
+	{
+		TArray<int32> OutBones;
+
+		if (LevelAttribute.IsValid() && SimulationTypeAttribute.IsValid())
+		{
+			const TManagedArray<int32>& Levels = LevelAttribute.Get();
+			bool bHasSimType = SimulationTypeAttribute.IsValid();
+
+			if (!bHasSimType)
+			{
+				// cannot filter by cluster/rigid if simulation type is not available
+				ensure(!bOnlyClusteredOrRigid);
+				bOnlyClusteredOrRigid = false;
+			}
+
+			bool bAllLevels = Level == INDEX_NONE;
+
+			const int32 NumTransforms = ParentAttribute.Num();
+			for (int32 BoneIdx = 0; BoneIdx < NumTransforms; BoneIdx++)
+			{
+				bool bIsRigid = bHasSimType && SimulationTypeAttribute[BoneIdx] == FGeometryCollection::ESimulationTypes::FST_Rigid;
+				bool bIsClustered = bHasSimType && SimulationTypeAttribute[BoneIdx] == FGeometryCollection::ESimulationTypes::FST_Clustered;
+				if (
+					// (if skipping embedded) sim type is clustered or rigid 
+					(!bOnlyClusteredOrRigid || bIsClustered || bIsRigid)
+					&&
+					// level is at or before the target
+					(bAllLevels || Levels[BoneIdx] == Level)
+					)
+				{
+					OutBones.Add(BoneIdx);
+				}
+			}
+		}
+
+		return OutBones;
+	}
+
 	void FCollectionTransformSelectionFacade::Sanitize(TArray<int32>& InOutSelection, bool bFavorParents) const
 	{
 		if (ParentAttribute.IsValid())
@@ -188,6 +227,37 @@ namespace GeometryCollection::Facades
 		InOutSelection = RigidSelection;
 	}
 
+	void FCollectionTransformSelectionFacade::ConvertEmbeddedSelectionToParents(TArray<int32>& InOutSelection) const
+	{
+		if (!SimulationTypeAttribute.IsValid()) // if no simulation type, no embedded to convert
+		{
+			return;
+		}
+
+		const TManagedArray<int32>& Parents = ParentAttribute.Get();
+		const TManagedArray<int32>& SimulationType = SimulationTypeAttribute.Get();
+
+		Sanitize(InOutSelection);
+
+		for (int32 SelBoneIdx = 0; SelBoneIdx < InOutSelection.Num(); ++SelBoneIdx)
+		{
+			int32 Bone = InOutSelection[SelBoneIdx];
+			if (SimulationType[Bone] == FGeometryCollection::ESimulationTypes::FST_None)
+			{
+				int32 Parent = Parents[Bone];
+				if (Parent != INDEX_NONE)
+				{
+					InOutSelection[SelBoneIdx] = Parent;
+				}
+				else // embedded should always have a parent, but if it somehow does not, just remove from selection
+				{
+					InOutSelection.RemoveAtSwap(SelBoneIdx, 1, false);
+					--SelBoneIdx; // reconsider swapped-in-element at this idx next iter
+				}
+			}
+		}
+	}
+
 	void FCollectionTransformSelectionFacade::ConvertSelectionToClusterNodes(TArray<int32>& InOutSelection, bool bLeaveRigidRoots) const
 	{
 		const TManagedArray<int32>& Parents = ParentAttribute.Get();
diff --git a/Engine/Source/Runtime/Experimental/Chaos/Private/GeometryCollection/GeometryCollectionClusteringUtility.cpp b/Engine/Source/Runtime/Experimental/Chaos/Private/GeometryCollection/GeometryCollectionClusteringUtility.cpp
--- a/Engine/Source/Runtime/Experimental/Chaos/Private/GeometryCollection/GeometryCollectionClusteringUtility.cpp
+++ b/Engine/Source/Runtime/Experimental/Chaos/Private/GeometryCollection/GeometryCollectionClusteringUtility.cpp
@@ -381,14 +381,13 @@ void FGeometryCollectionClusteringUtility::RenameBone(FGeometryCollection* Geome
 	}
 }
 
-int32 FGeometryCollectionClusteringUtility::PickBestNodeToMergeTo(const FGeometryCollection* GeometryCollection, const TArray<int32>& SourceElements)
+int32 FGeometryCollectionClusteringUtility::PickBestNodeToMergeTo(const FManagedArrayCollection* Collection, const TArray<int32>& SourceElements)
 {
-	if (!ensure(GeometryCollection->HasAttribute("Level", FGeometryCollection::TransformGroup)))
+	const Chaos::Facades::FCollectionHierarchyFacade HierarchyFacade(*Collection);
+	if (!HierarchyFacade.IsValid() || !HierarchyFacade.HasLevelAttribute())
 	{
 		return -1;
 	}
-	const TManagedArray<int32>& Levels = GeometryCollection->GetAttribute<int32>("Level", FGeometryCollection::TransformGroup);
-	const TManagedArray<TSet<int32>>& Children = GeometryCollection->Children;
 
 	// which of the source elements is the most significant, closest to the root that has children (is a cluster)
 	int32 ElementClosestToRoot = -1;
@@ -396,9 +395,11 @@ int32 FGeometryCollectionClusteringUtility::PickBestNodeToMergeTo(const FGeometr
 
 	for (int32 Element : SourceElements)
 	{
-		if (Children[Element].Num() > 0 && (Levels[Element] < LevelClosestToRoot || LevelClosestToRoot == -1))
+		const TSet<int32>* Children = HierarchyFacade.FindChildren(Element);
+		int32 Level = HierarchyFacade.GetInitialLevel(Element);
+		if (Children && !Children->IsEmpty() && (Level < LevelClosestToRoot || LevelClosestToRoot == -1))
 		{
-			LevelClosestToRoot = Levels[Element];
+			LevelClosestToRoot = Level;
 			ElementClosestToRoot = Element;
 		}
 	}
@@ -874,7 +875,7 @@ void FGeometryCollectionClusteringUtility::MoveUpOneHierarchyLevel(FGeometryColl
 }
 
 
-int32 FGeometryCollectionClusteringUtility::FindLowestCommonAncestor(FGeometryCollection* GeometryCollection, const TArray<int32>& SelectedBones)
+int32 FGeometryCollectionClusteringUtility::FindLowestCommonAncestor(const FManagedArrayCollection* Collection, const TArray<int32>& SelectedBones)
 {
 	const int32 SelectionCount = SelectedBones.Num();
 	if (SelectionCount == 0)
@@ -889,21 +890,25 @@ int32 FGeometryCollectionClusteringUtility::FindLowestCommonAncestor(FGeometryCo
 		{
 			return INDEX_NONE;
 		}
-		LCA = FindLowestCommonAncestor(GeometryCollection, LCA, SelectedBones[Index]);
+		LCA = FindLowestCommonAncestor(Collection, LCA, SelectedBones[Index]);
 	}
 	return LCA;
 }
 
-int32 FGeometryCollectionClusteringUtility::FindLowestCommonAncestor(FGeometryCollection* GeometryCollection, int32 N0, int32 N1)
+int32 FGeometryCollectionClusteringUtility::FindLowestCommonAncestor(const FManagedArrayCollection* Collection, int32 N0, int32 N1)
 {
-	const TManagedArray<int32>& Parent = GeometryCollection->GetAttribute<int32>("Parent", FGeometryCollection::TransformGroup);
+	const TManagedArray<int32>* Parent = Collection->FindAttribute<int32>("Parent", FGeometryCollection::TransformGroup);
+	if (!Parent)
+	{
+		return INDEX_NONE;
+	}
 
 	// Record the path to root from the first 
 	TArray<int32> PathToRoot0;
 	PathToRoot0.Add(N0);
 	while (PathToRoot0.Last() != INDEX_NONE)
 	{
-		PathToRoot0.Add(Parent[PathToRoot0.Last()]);
+		PathToRoot0.Add((*Parent)[PathToRoot0.Last()]);
 	}
 
 	// Traverse from the second node to root and return the first node found that is in the first path.
@@ -914,7 +919,7 @@ int32 FGeometryCollectionClusteringUtility::FindLowestCommonAncestor(FGeometryCo
 		{
 			return LCA;
 		}
-		LCA = Parent[LCA];
+		LCA = (*Parent)[LCA];
 	}
 
 	// No common ancestor
diff --git a/Engine/Source/Runtime/Experimental/Chaos/Private/GeometryCollection/GeometryCollectionUtility.cpp b/Engine/Source/Runtime/Experimental/Chaos/Private/GeometryCollection/GeometryCollectionUtility.cpp
--- a/Engine/Source/Runtime/Experimental/Chaos/Private/GeometryCollection/GeometryCollectionUtility.cpp
+++ b/Engine/Source/Runtime/Experimental/Chaos/Private/GeometryCollection/GeometryCollectionUtility.cpp
@@ -478,7 +478,7 @@ namespace GeometryCollection
 	}
 
 
-	void GenerateTemporaryGuids(FTransformCollection* Collection, int32 StartIdx, bool bForceInit)
+	void GenerateTemporaryGuids(FManagedArrayCollection* Collection, int32 StartIdx, bool bForceInit)
 	{
 		bool bNeedsInit = false;
 		if (!Collection->HasAttribute("GUID", FTransformCollection::TransformGroup))
diff --git a/Engine/Source/Runtime/Experimental/Chaos/Public/GeometryCollection/Facades/CollectionHierarchyFacade.h b/Engine/Source/Runtime/Experimental/Chaos/Public/GeometryCollection/Facades/CollectionHierarchyFacade.h
--- a/Engine/Source/Runtime/Experimental/Chaos/Public/GeometryCollection/Facades/CollectionHierarchyFacade.h
+++ b/Engine/Source/Runtime/Experimental/Chaos/Public/GeometryCollection/Facades/CollectionHierarchyFacade.h
@@ -47,6 +47,9 @@ namespace Chaos::Facades
 		/** Get direct children from a specific transform index. Assumes parent attribute is valid. */
 		TArray<int32> GetChildrenAsArray(int32 TransformIndex) const;
 
+		/** Get set of children from a transform index, or nullptr if not available */
+		const TSet<int32>* FindChildren(int32 TransformIndex) const;
+
 		/** Get parent of a specific transform index */
 		int32 GetParent(int32 TransformIndex) const;
 
diff --git a/Engine/Source/Runtime/Experimental/Chaos/Public/GeometryCollection/Facades/CollectionTransformSelectionFacade.h b/Engine/Source/Runtime/Experimental/Chaos/Public/GeometryCollection/Facades/CollectionTransformSelectionFacade.h
--- a/Engine/Source/Runtime/Experimental/Chaos/Public/GeometryCollection/Facades/CollectionTransformSelectionFacade.h
+++ b/Engine/Source/Runtime/Experimental/Chaos/Public/GeometryCollection/Facades/CollectionTransformSelectionFacade.h
@@ -48,9 +48,12 @@ namespace GeometryCollection::Facades
 		/**  */
 		bool HasSelectedAncestor(const TArray<int32>& InSelection, const int32 Index) const;
 
-		/**  */
+		/** Get bones filtered by the level as they would be in the editor; i.e., the bones with level <= the target Level, but with some additional logic for filtering, etc */
 		TArray<int32> GetBonesByLevel(const int32 Level, bool bOnlyClusteredOrRigid, bool bSkipFiltered) const;
 
+		/** Get bones exactly at a target level. A target level of -1 will return all bones. */
+		TArray<int32> GetBonesExactlyAtLevel(const int32 TargetLevel, bool bOnlyClusteredOrRigid) const;
+
 		/**  */
 		void Sanitize(TArray<int32>& InSelection, bool bFavorParents = true) const;
 
@@ -64,6 +67,8 @@ namespace GeometryCollection::Facades
 		/** Replace any Embedded or Rigid transforms in InOutSelection with their parent cluster. If there is just a rigid root node selected, optionally keep it if bLeaveRigidRoot is true. */
 		void ConvertSelectionToClusterNodes(TArray<int32>& InOutSelection, bool bLeaveRigidRoots = true) const;
 
+		void ConvertEmbeddedSelectionToParents(TArray<int32>& InOutSelection) const;
+
 		/**  */
 		bool CanSelectRootBones() const { return ParentAttribute.IsValid(); }
 		TArray<int32> SelectRootBones() const;
diff --git a/Engine/Source/Runtime/Experimental/Chaos/Public/GeometryCollection/GeometryCollectionClusteringUtility.h b/Engine/Source/Runtime/Experimental/Chaos/Public/GeometryCollection/GeometryCollectionClusteringUtility.h
--- a/Engine/Source/Runtime/Experimental/Chaos/Public/GeometryCollection/GeometryCollectionClusteringUtility.h
+++ b/Engine/Source/Runtime/Experimental/Chaos/Public/GeometryCollection/GeometryCollectionClusteringUtility.h
@@ -122,7 +122,7 @@ public:
 	static void MoveUpOneHierarchyLevel(FGeometryCollection* GeometryCollection, const TArray<int32>& SelectedBones);
 
 	/** Find the lowest common ancestor index of currently selected nodes. Returns INDEX_NODE if there is no common ancestor. */
-	static int32 FindLowestCommonAncestor(FGeometryCollection* GeometryCollection, const TArray<int32>& SelectedBones);
+	static int32 FindLowestCommonAncestor(const FManagedArrayCollection* GeometryCollection, const TArray<int32>& SelectedBones);
 
 	/** Delete any cluster nodes discovered to have no children. Returns true if any nodes were removed. */
 	static bool RemoveDanglingClusters(FGeometryCollection* GeometryCollection);
@@ -132,7 +132,7 @@ public:
 
 	static void ValidateResults(FGeometryCollection* GeometryCollection);
 
-	static int32 PickBestNodeToMergeTo(const FGeometryCollection* GeometryCollection, const TArray<int32>& SourceElements);
+	static int32 PickBestNodeToMergeTo(const FManagedArrayCollection* Collection, const TArray<int32>& SourceElements);
 
 private:
 	// #todo: intend to remove reliance on custom attributes for slider by making use of Rest/Dynamic collections
@@ -140,5 +140,5 @@ private:
 	
 	static void RecursivelyUpdateHierarchyLevelOfChildren(TManagedArray<int32>& Levels, const TManagedArray<TSet<int32>>& Children, int32 ParentElement);
 
-	static int32 FindLowestCommonAncestor(FGeometryCollection* GeometryCollection, int32 N0, int32 N1);
+	static int32 FindLowestCommonAncestor(const FManagedArrayCollection* GeometryCollection, int32 N0, int32 N1);
 };
diff --git a/Engine/Source/Runtime/Experimental/Chaos/Public/GeometryCollection/GeometryCollectionUtility.h b/Engine/Source/Runtime/Experimental/Chaos/Public/GeometryCollection/GeometryCollectionUtility.h
--- a/Engine/Source/Runtime/Experimental/Chaos/Public/GeometryCollection/GeometryCollectionUtility.h
+++ b/Engine/Source/Runtime/Experimental/Chaos/Public/GeometryCollection/GeometryCollectionUtility.h
@@ -165,7 +165,7 @@ namespace GeometryCollection
 	*/
 	void 
 	CHAOS_API
-	GenerateTemporaryGuids(FTransformCollection* Collection, int32 StartIdx = 0, bool bForceInit=false);
+	GenerateTemporaryGuids(FManagedArrayCollection* Collection, int32 StartIdx = 0, bool bForceInit=false);
 
 	/***
 	* Compute inner and outer radius from a set of vertices 
