add dataflow node to simplify existing convex hulls on a collection

diff --git a/Engine/Plugins/Experimental/Fracture/Source/FractureEngine/Private/FractureEngineConvex.cpp b/Engine/Plugins/Experimental/Fracture/Source/FractureEngine/Private/FractureEngineConvex.cpp
new file mode 100644
--- /dev/null
+++ b/Engine/Plugins/Experimental/Fracture/Source/FractureEngine/Private/FractureEngineConvex.cpp
@@ -0,0 +1,151 @@
+// Copyright Epic Games, Inc. All Rights Reserved.
+
+#include "FractureEngineConvex.h"
+#include "Chaos/Convex.h"
+#include "DynamicMesh/DynamicMesh3.h"
+#include "DynamicMesh/DynamicMeshAABBTree3.h"
+#include "ProjectionTargets.h"
+#include "MeshSimplification.h"
+#include "GeometryCollection/GeometryCollection.h"
+#include "GeometryCollection/GeometryCollectionConvexUtility.h"
+
+namespace UE::FractureEngine::Convex
+{
+	bool SimplifyConvexHulls(FManagedArrayCollection& Collection, const FSimplifyHullSettings& Settings, bool bRestrictToSelection, const TArrayView<const int32> TransformSelection)
+	{
+		if (!FGeometryCollectionConvexUtility::HasConvexHullData(&Collection))
+		{
+			// nothing to simplify
+			return false;
+		}
+		
+		TManagedArray<TSet<int32>>& TransformToConvexInds = Collection.ModifyAttribute<TSet<int32>>("TransformToConvexIndices", FTransformCollection::TransformGroup);
+		TManagedArray<TUniquePtr<::Chaos::FConvex>>& ConvexHulls = Collection.ModifyAttribute<TUniquePtr<::Chaos::FConvex>>("ConvexHull", "Convex");
+		
+		auto SimplifyBone = [&TransformToConvexInds, &ConvexHulls, &Settings](int32 BoneIdx) -> bool
+		{
+			if (BoneIdx < 0 || BoneIdx >= TransformToConvexInds.Num())
+			{
+				// invalid bone index
+				return false;
+			}
+			bool bNoFailures = true;
+			for (int32 ConvexIdx : TransformToConvexInds[BoneIdx])
+			{
+				bool bSuccess = SimplifyConvexHull(ConvexHulls[ConvexIdx].Get(), ConvexHulls[ConvexIdx].Get(), Settings);
+				bNoFailures = bNoFailures && bSuccess;
+			}
+			return bNoFailures;
+		};
+
+		bool bNoFailures = true;
+
+		if (bRestrictToSelection)
+		{
+			for (int32 BoneIdx : TransformSelection)
+			{
+				bool bSuccess = SimplifyBone(BoneIdx);
+				bNoFailures = bNoFailures && bSuccess;
+			}
+		}
+		else
+		{
+			for (int32 BoneIdx = 0; BoneIdx < Collection.NumElements(FGeometryCollection::TransformGroup); ++BoneIdx)
+			{
+				bool bSuccess = SimplifyBone(BoneIdx);
+				bNoFailures = bNoFailures && bSuccess;
+			}
+		}
+
+		return bNoFailures;
+	}
+
+	bool SimplifyConvexHull(const ::Chaos::FConvex* InConvexHull, ::Chaos::FConvex* OutConvexHull, const FSimplifyHullSettings& Settings)
+	{
+		if (!InConvexHull || !OutConvexHull || !InConvexHull->HasStructureData())
+		{
+			return false;
+		}
+
+		const ::Chaos::FConvexStructureData& ConvexStructure = InConvexHull->GetStructureData();
+		const int32 NumV = InConvexHull->NumVertices();
+		const int32 NumP = InConvexHull->NumPlanes();
+
+		// Check if no simplification required, and skip simplification in that case
+		int32 ExpectNumT = 0;
+		for (int32 PIdx = 0; PIdx < NumP; ++PIdx)
+		{
+			const int32 NumFaceV = ConvexStructure.NumPlaneVertices(PIdx);
+			ExpectNumT += FMath::Max(0, NumFaceV - 2);
+		}
+		if (Settings.bUseTargetTriangleCount && ExpectNumT <= Settings.TargetTriangleCount)
+		{
+			if (OutConvexHull != InConvexHull)
+			{
+				*OutConvexHull = MoveTemp(*InConvexHull->CopyAsConvex());
+			}
+			return true;
+		}
+	
+		// Convert to DynamicMesh to run simplifier
+		UE::Geometry::FDynamicMesh3 Mesh;
+		for (int32 VIdx = 0; VIdx < NumV; ++VIdx)
+		{
+			int32 MeshVIdx = Mesh.AppendVertex((FVector3d)InConvexHull->GetVertex(VIdx));
+			checkSlow(MeshVIdx == VIdx);
+		}
+		for (int32 PIdx = 0; PIdx < NumP; ++PIdx)
+		{
+			const int32 NumFaceV = ConvexStructure.NumPlaneVertices(PIdx);
+			const int32 V0 = ConvexStructure.GetPlaneVertex(PIdx, 0);
+			for (int32 SubIdx = 1; SubIdx + 1 < NumFaceV; ++SubIdx)
+			{
+				const int32 V1 = ConvexStructure.GetPlaneVertex(PIdx, SubIdx);
+				const int32 V2 = ConvexStructure.GetPlaneVertex(PIdx, SubIdx + 1);
+				Mesh.AppendTriangle(UE::Geometry::FIndex3i(V0, V1, V2));
+			}
+		}
+		
+		// Run simplification
+		UE::Geometry::FVolPresMeshSimplification Simplifier(&Mesh);
+		
+		Simplifier.CollapseMode =
+			Settings.bUseExistingVertexPositions ?
+				UE::Geometry::FVolPresMeshSimplification::ESimplificationCollapseModes::MinimalExistingVertexError
+			:	UE::Geometry::FVolPresMeshSimplification::ESimplificationCollapseModes::MinimalQuadricPositionError;
+		if (Settings.bUseGeometricTolerance)
+		{
+			Simplifier.GeometricErrorConstraint = UE::Geometry::FVolPresMeshSimplification::EGeometricErrorCriteria::PredictedPointToProjectionTarget;
+			Simplifier.GeometricErrorTolerance = Settings.ErrorTolerance;
+		}
+		if (Settings.bUseGeometricTolerance)
+		{
+			// Simplify to the smallest non-degenerate number of triangles, relying on geometric error criteria
+			UE::Geometry::FDynamicMesh3 ProjectionTargetMesh(Mesh);
+			UE::Geometry::FDynamicMeshAABBTree3 ProjectionTargetSpatial(&ProjectionTargetMesh, true);
+			UE::Geometry::FMeshProjectionTarget ProjTarget(&ProjectionTargetMesh, &ProjectionTargetSpatial);
+			Simplifier.SetProjectionTarget(&ProjTarget);
+			int32 TargetTriCount = Settings.bUseTargetTriangleCount ? Settings.TargetTriangleCount : 4;
+			Simplifier.SimplifyToTriangleCount(TargetTriCount);
+		}
+		else if (Settings.bUseTargetTriangleCount)
+		{
+			Simplifier.SimplifyToTriangleCount(Settings.TargetTriangleCount);
+		}
+		else
+		{
+			// Note: Quadric error threshold doesn't have the same geometric meaning as distance; this is not equivalent to using a geometric error tolerance
+			Simplifier.SimplifyToMaxError(Settings.ErrorTolerance * Settings.ErrorTolerance);
+		}
+
+		TArray<::Chaos::FVec3f> NewConvexVerts;
+		NewConvexVerts.Reserve(Mesh.VertexCount());
+		for (int32 VIdx : Mesh.VertexIndicesItr())
+		{
+			NewConvexVerts.Add((::Chaos::FVec3f)Mesh.GetVertex(VIdx));
+		}
+		*OutConvexHull = ::Chaos::FConvex(NewConvexVerts, InConvexHull->GetMargin(), ::Chaos::FConvexBuilder::EBuildMethod::Default);
+
+		return true;
+	}
+}
diff --git a/Engine/Plugins/Experimental/Fracture/Source/FractureEngine/Public/FractureEngineConvex.h b/Engine/Plugins/Experimental/Fracture/Source/FractureEngine/Public/FractureEngineConvex.h
new file mode 100644
--- /dev/null
+++ b/Engine/Plugins/Experimental/Fracture/Source/FractureEngine/Public/FractureEngineConvex.h
@@ -0,0 +1,37 @@
+// Copyright Epic Games, Inc. All Rights Reserved.
+
+#pragma once
+
+#include "Containers/ArrayView.h"
+
+namespace Chaos { class FConvex; }
+struct FManagedArrayCollection;
+
+namespace UE::FractureEngine::Convex
+{
+	// Options to control the simplification of an existing convex hull
+	struct FSimplifyHullSettings
+	{
+		bool bUseGeometricTolerance = true;
+		double ErrorTolerance = 5;
+
+		bool bUseTargetTriangleCount = false;
+		int32 TargetTriangleCount = 20;
+
+		bool bUseExistingVertexPositions = true;
+	};
+
+	/**
+	 * Simplify the convex hulls on the given Collection. Optionally only simplify the hulls on the transforms in TransformSelection.
+	 * @return true if the collection has convex hulls and the hulls were either simplified or did not need to be
+	 */
+	bool FRACTUREENGINE_API SimplifyConvexHulls(FManagedArrayCollection& Collection, const FSimplifyHullSettings& Settings, bool bRestrictToSelection = false, const TArrayView<const int32> TransformSelection = TArrayView<const int32>());
+
+	/**
+	 * Simplify a convex hull using the given Settings.OutConvexHull can optionally be the same pointer as InConvexHull.
+	 * @return true if the hull had valid data and either was simplified or did not need to be (e.g. already had few enough triangles)
+	 */
+	bool FRACTUREENGINE_API SimplifyConvexHull(const ::Chaos::FConvex* InConvexHull, ::Chaos::FConvex* OutConvexHull, const FSimplifyHullSettings& Settings);
+
+}
+
diff --git a/Engine/Plugins/Experimental/GeometryCollectionPlugin/Source/GeometryCollectionNodes/Private/Dataflow/GeometryCollectionUtilityNodes.cpp b/Engine/Plugins/Experimental/GeometryCollectionPlugin/Source/GeometryCollectionNodes/Private/Dataflow/GeometryCollectionUtilityNodes.cpp
--- a/Engine/Plugins/Experimental/GeometryCollectionPlugin/Source/GeometryCollectionNodes/Private/Dataflow/GeometryCollectionUtilityNodes.cpp
+++ b/Engine/Plugins/Experimental/GeometryCollectionPlugin/Source/GeometryCollectionNodes/Private/Dataflow/GeometryCollectionUtilityNodes.cpp
@@ -7,6 +7,11 @@
 #include "GeometryCollection/GeometryCollection.h"
 #include "GeometryCollection/GeometryCollectionConvexUtility.h"
 
+#include "FractureEngineConvex.h"
+
+#include "MeshSimplification.h"
+
+
 //#include UE_INLINE_GENERATED_CPP_BY_NAME(GeometryCollectionUtilityNodes)
 
 namespace Dataflow
@@ -17,6 +22,7 @@ namespace Dataflow
 		static const FLinearColor CDefaultNodeBodyTintColor = FLinearColor(0.f, 0.f, 0.f, 0.5f);
 
 		DATAFLOW_NODE_REGISTER_CREATION_FACTORY(FCreateLeafConvexHullsDataflowNode);
+		DATAFLOW_NODE_REGISTER_CREATION_FACTORY(FSimplifyConvexHullsDataflowNode);
 		DATAFLOW_NODE_REGISTER_CREATION_FACTORY(FCreateNonOverlappingConvexHullsDataflowNode);
 		DATAFLOW_NODE_REGISTER_CREATION_FACTORY(FGenerateClusterConvexHullsFromLeafHullsDataflowNode);
 		DATAFLOW_NODE_REGISTER_CREATION_FACTORY(FGenerateClusterConvexHullsFromChildrenHullsDataflowNode);
@@ -65,6 +71,47 @@ void FCreateLeafConvexHullsDataflowNode::Evaluate(Dataflow::FContext& Context, c
 	}
 }
 
+FSimplifyConvexHullsDataflowNode::FSimplifyConvexHullsDataflowNode(const Dataflow::FNodeParameters& InParam, FGuid InGuid)
+	: FDataflowNode(InParam, InGuid)
+{
+	RegisterInputConnection(&Collection);
+	RegisterInputConnection(&OptionalSelectionFilter);
+	RegisterInputConnection(&SimplificationDistanceThreshold);
+	RegisterInputConnection(&MinTargetTriangleCount);
+	RegisterOutputConnection(&Collection);
+}
+
+void FSimplifyConvexHullsDataflowNode::Evaluate(Dataflow::FContext& Context, const FDataflowOutput* Out) const
+{
+	if (Out->IsA(&Collection) && IsConnected(&Collection))
+	{
+		FManagedArrayCollection InCollection = GetValue(Context, &Collection);
+		if (InCollection.NumElements(FGeometryCollection::TransformGroup) == 0)
+		{
+			SetValue<FManagedArrayCollection>(Context, MoveTemp(InCollection), &Collection);
+			return;
+		}
+
+		TArray<int32> SelectedBones;
+		bool bRestrictToSelection = false;
+		if (IsConnected(&OptionalSelectionFilter))
+		{
+			const FDataflowTransformSelection& InOptionalSelectionFilter = GetValue(Context, &OptionalSelectionFilter);
+			bRestrictToSelection = true;
+			SelectedBones = InOptionalSelectionFilter.AsArray();
+		}
+
+		UE::FractureEngine::Convex::FSimplifyHullSettings Settings;
+		Settings.ErrorTolerance = GetValue(Context, &SimplificationDistanceThreshold);
+		Settings.bUseGeometricTolerance = true;
+		Settings.bUseTargetTriangleCount = true;
+		Settings.bUseExistingVertexPositions = bUseExistingVertices;
+		Settings.TargetTriangleCount = GetValue(Context, &MinTargetTriangleCount);
+		UE::FractureEngine::Convex::SimplifyConvexHulls(InCollection, Settings, bRestrictToSelection, SelectedBones);
+		SetValue<FManagedArrayCollection>(Context, MoveTemp(InCollection), &Collection);
+	}
+}
+
 FCreateNonOverlappingConvexHullsDataflowNode::FCreateNonOverlappingConvexHullsDataflowNode(const Dataflow::FNodeParameters& InParam, FGuid InGuid)
 	: FDataflowNode(InParam, InGuid)
 {
diff --git a/Engine/Plugins/Experimental/GeometryCollectionPlugin/Source/GeometryCollectionNodes/Public/Dataflow/GeometryCollectionUtilityNodes.h b/Engine/Plugins/Experimental/GeometryCollectionPlugin/Source/GeometryCollectionNodes/Public/Dataflow/GeometryCollectionUtilityNodes.h
--- a/Engine/Plugins/Experimental/GeometryCollectionPlugin/Source/GeometryCollectionNodes/Public/Dataflow/GeometryCollectionUtilityNodes.h
+++ b/Engine/Plugins/Experimental/GeometryCollectionPlugin/Source/GeometryCollectionNodes/Public/Dataflow/GeometryCollectionUtilityNodes.h
@@ -59,6 +59,37 @@ public:
 	virtual void Evaluate(Dataflow::FContext& Context, const FDataflowOutput* Out) const override;
 };
 
+USTRUCT(meta = (DataflowGeometryCollection))
+struct FSimplifyConvexHullsDataflowNode : public FDataflowNode
+{
+	GENERATED_USTRUCT_BODY()
+	DATAFLOW_NODE_DEFINE_INTERNAL(FSimplifyConvexHullsDataflowNode, "SimplifyConvexHulls", "GeometryCollection|Utilities", "")
+
+public:
+	UPROPERTY(meta = (DataflowInput, DataflowOutput))
+	FManagedArrayCollection Collection;
+
+	/** Optional transform selection to compute leaf hulls on -- if not provided, all leaf hulls will be computed. */
+	UPROPERTY(meta = (DataflowInput, DataflowIntrinsic))
+	FDataflowTransformSelection OptionalSelectionFilter;
+
+	/** Computed convex hulls are simplified to keep points spaced at least this far apart (except where needed to keep the hull from collapsing to zero volume) */
+	UPROPERTY(EditAnywhere, Category = "Convex", meta = (DataflowInput, ClampMin = 0.f))
+	float SimplificationDistanceThreshold = 10.f;
+
+	/** The minimum number of triangles to use for the convex hull. Will use more triangles than this as-needed to preserve the distance threshold. */
+	UPROPERTY(EditAnywhere, Category = "Convex", meta = (DataflowInput, ClampMin = 4))
+	int32 MinTargetTriangleCount = 12;
+
+	/** Whether to restrict the simplified hulls to only use vertices from the original hulls. */
+	UPROPERTY(EditAnywhere, Category = "Convex")
+	bool bUseExistingVertices = false;
+
+	FSimplifyConvexHullsDataflowNode(const Dataflow::FNodeParameters& InParam, FGuid InGuid = FGuid::NewGuid());
+
+	virtual void Evaluate(Dataflow::FContext& Context, const FDataflowOutput* Out) const override;
+};
+
 /**
  *
  * Generates convex hull representation for the bones for simulation
diff --git a/Engine/Source/Runtime/Experimental/Chaos/Private/GeometryCollection/GeometryCollectionConvexUtility.cpp b/Engine/Source/Runtime/Experimental/Chaos/Private/GeometryCollection/GeometryCollectionConvexUtility.cpp
--- a/Engine/Source/Runtime/Experimental/Chaos/Private/GeometryCollection/GeometryCollectionConvexUtility.cpp
+++ b/Engine/Source/Runtime/Experimental/Chaos/Private/GeometryCollection/GeometryCollectionConvexUtility.cpp
@@ -134,9 +134,9 @@ TOptional<FGeometryCollectionConvexUtility::FGeometryCollectionConvexData> FGeom
 	return TOptional<FGeometryCollectionConvexUtility::FGeometryCollectionConvexData>(ConvexData);
 }
 
-bool FGeometryCollectionConvexUtility::HasConvexHullData(FGeometryCollection* GeometryCollection)
+bool FGeometryCollectionConvexUtility::HasConvexHullData(const FManagedArrayCollection* Collection)
 {
-	return GeometryCollection->HasAttribute("TransformToConvexIndices", FTransformCollection::TransformGroup) && GeometryCollection->HasAttribute("ConvexHull", "Convex");
+	return Collection->HasAttribute("TransformToConvexIndices", FTransformCollection::TransformGroup) && Collection->HasAttribute("ConvexHull", "Convex");
 }
 
 FGeometryCollectionConvexUtility::FGeometryCollectionConvexData FGeometryCollectionConvexUtility::GetValidConvexHullData(FGeometryCollection* GeometryCollection)
diff --git a/Engine/Source/Runtime/Experimental/Chaos/Public/GeometryCollection/GeometryCollectionConvexUtility.h b/Engine/Source/Runtime/Experimental/Chaos/Public/GeometryCollection/GeometryCollectionConvexUtility.h
--- a/Engine/Source/Runtime/Experimental/Chaos/Public/GeometryCollection/GeometryCollectionConvexUtility.h
+++ b/Engine/Source/Runtime/Experimental/Chaos/Public/GeometryCollection/GeometryCollectionConvexUtility.h
@@ -83,7 +83,7 @@ public:
 	static TOptional<FGeometryCollectionConvexData> GetConvexHullDataIfPresent(FGeometryCollection* GeometryCollection);
 
 	/** @return true if convex hull data is present */
-	static bool HasConvexHullData(FGeometryCollection* GeometryCollection);
+	static bool HasConvexHullData(const FManagedArrayCollection* GeometryCollection);
 
 	/**
 	 Create non-overlapping convex hull data for all transforms in the geometry collection (except transforms where it would be better to just use the hulls of the children) 
