make autocluster always update the volume attribute before using it

diff --git a/Engine/Plugins/Experimental/Fracture/Source/FractureEngine/Private/FractureEngineClustering.cpp b/Engine/Plugins/Experimental/Fracture/Source/FractureEngine/Private/FractureEngineClustering.cpp
--- a/Engine/Plugins/Experimental/Fracture/Source/FractureEngine/Private/FractureEngineClustering.cpp
+++ b/Engine/Plugins/Experimental/Fracture/Source/FractureEngine/Private/FractureEngineClustering.cpp
@@ -84,10 +84,8 @@ void FVoronoiPartitioner::MergeSmallPartitions(FGeometryCollection* GeometryColl
 	{
 		GenerateConnectivity(GeometryCollection);
 	}
-	if (!GeometryCollection->HasAttribute("Volume", FTransformCollection::TransformGroup))
-	{
-		FGeometryCollectionConvexUtility::SetVolumeAttributes(GeometryCollection);
-	}
+
+	FGeometryCollectionConvexUtility::SetVolumeAttributes(GeometryCollection);
 	const TManagedArray<float>& Volumes = GeometryCollection->GetAttribute<float>("Volume", FTransformCollection::TransformGroup);
 
 	float VolumeThreshold = SizeThreshold * SizeThreshold * SizeThreshold;
