add geometry collection debug nodes to support saving OBJ files of meshes, and converting convex hulls and sphere coverings to meshes
+ add sphere covering as an output for the 'generate cluster convex hull' nodes, so it can be visualized
 + fix convex hull meshes missing some triangles due to non-manifold mesh data (which occurs due to the aggressive MergeFaces routine in the physics convex hull build)

diff --git a/Engine/Plugins/Experimental/Fracture/Source/FractureEngine/Private/FractureEngineConvex.cpp b/Engine/Plugins/Experimental/Fracture/Source/FractureEngine/Private/FractureEngineConvex.cpp
--- a/Engine/Plugins/Experimental/Fracture/Source/FractureEngine/Private/FractureEngineConvex.cpp
+++ b/Engine/Plugins/Experimental/Fracture/Source/FractureEngine/Private/FractureEngineConvex.cpp
@@ -16,7 +16,7 @@ namespace
 {
 	// Local helpers for converting convex hulls to dynamic meshes, used to run geometry processing tasks on convex hulls (e.g., simplification, computing negative space)
 
-	static void AppendConvexHullToCompactDynamicMesh(const ::Chaos::FConvex* InConvexHull, UE::Geometry::FDynamicMesh3& Mesh, FTransform* OptionalTransform = nullptr)
+	static void AppendConvexHullToCompactDynamicMesh(const ::Chaos::FConvex* InConvexHull, UE::Geometry::FDynamicMesh3& Mesh, FTransform* OptionalTransform = nullptr, bool bFixNonmanifoldWithDuplicates = false, bool bInvertFaces = false)
 	{
 		check(Mesh.IsCompact());
 
@@ -40,9 +40,24 @@ namespace
 			const int32 V0 = StartV + ConvexStructure.GetPlaneVertex(PIdx, 0);
 			for (int32 SubIdx = 1; SubIdx + 1 < NumFaceV; ++SubIdx)
 			{
-				const int32 V1 = StartV + ConvexStructure.GetPlaneVertex(PIdx, SubIdx);
-				const int32 V2 = StartV + ConvexStructure.GetPlaneVertex(PIdx, SubIdx + 1);
-				Mesh.AppendTriangle(UE::Geometry::FIndex3i(V0, V1, V2));
+				int32 V1 = StartV + ConvexStructure.GetPlaneVertex(PIdx, SubIdx);
+				int32 V2 = StartV + ConvexStructure.GetPlaneVertex(PIdx, SubIdx + 1);
+				if (bInvertFaces)
+				{
+					Swap(V1, V2);
+				}
+				int32 ResultTID = Mesh.AppendTriangle(UE::Geometry::FIndex3i(V0, V1, V2));
+				if (bFixNonmanifoldWithDuplicates && ResultTID == UE::Geometry::FDynamicMesh3::NonManifoldID)
+				{
+					// failed to append due to a non-manifold triangle; try adding all the vertices independently so we at least capture the shape
+					// note: this should not happen for normal convex hulls, but the current convex hull algorithm does some aggressive face merging that sometimes creates weird geometry
+					UE::Geometry::FIndex3i DuplicateVerts(
+						Mesh.AppendVertex(Mesh.GetVertex(V0)),
+						Mesh.AppendVertex(Mesh.GetVertex(V1)),
+						Mesh.AppendVertex(Mesh.GetVertex(V2))
+					);
+					Mesh.AppendTriangle(DuplicateVerts);
+				}
 			}
 		}
 	}
@@ -57,6 +72,60 @@ namespace
 
 namespace UE::FractureEngine::Convex
 {
+	bool GetConvexHullsAsDynamicMesh(const FManagedArrayCollection& Collection, UE::Geometry::FDynamicMesh3& OutMesh, bool bRestrictToSelection, const TArrayView<const int32> TransformSelection)
+	{
+		OutMesh.Clear();
+
+		if (!FGeometryCollectionConvexUtility::HasConvexHullData(&Collection))
+		{
+			// nothing to append
+			return false;
+		}
+
+		const TManagedArray<TSet<int32>>& TransformToConvexInds = Collection.GetAttribute<TSet<int32>>("TransformToConvexIndices", FTransformCollection::TransformGroup);
+		const TManagedArray<TUniquePtr<::Chaos::FConvex>>& ConvexHulls = Collection.GetAttribute<TUniquePtr<::Chaos::FConvex>>("ConvexHull", "Convex");
+
+		GeometryCollection::Facades::FCollectionTransformFacade TransformFacade(Collection);
+		TArray<FTransform> GlobalTransformArray = TransformFacade.ComputeCollectionSpaceTransforms();
+
+		auto AppendBone = [&TransformToConvexInds, &ConvexHulls, &GlobalTransformArray, &OutMesh](int32 BoneIdx) -> bool
+		{
+			if (BoneIdx < 0 || BoneIdx >= TransformToConvexInds.Num())
+			{
+				// invalid bone index
+				return false;
+			}
+			for (int32 ConvexIdx : TransformToConvexInds[BoneIdx])
+			{
+				constexpr bool bConvertNonManifold = true; // Add non-manifold faces so they are still included in the debug visualization
+				constexpr bool bInvertFaces = true; // FConvex mesh data appears to have opposite default winding from what we expect for triangle meshes
+				AppendConvexHullToCompactDynamicMesh(ConvexHulls[ConvexIdx].Get(), OutMesh, &GlobalTransformArray[BoneIdx], bConvertNonManifold, bInvertFaces);
+			}
+			return true;
+		};
+
+		bool bNoFailures = true;
+
+		if (bRestrictToSelection)
+		{
+			for (int32 BoneIdx : TransformSelection)
+			{
+				bool bSuccess = AppendBone(BoneIdx);
+				bNoFailures = bNoFailures && bSuccess;
+			}
+		}
+		else
+		{
+			for (int32 BoneIdx = 0; BoneIdx < Collection.NumElements(FGeometryCollection::TransformGroup); ++BoneIdx)
+			{
+				bool bSuccess = AppendBone(BoneIdx);
+				bNoFailures = bNoFailures && bSuccess;
+			}
+		}
+
+		return bNoFailures;
+	}
+
 	bool SimplifyConvexHulls(FManagedArrayCollection& Collection, const FSimplifyHullSettings& Settings, bool bRestrictToSelection, const TArrayView<const int32> TransformSelection)
 	{
 		if (!FGeometryCollectionConvexUtility::HasConvexHullData(&Collection))
@@ -215,7 +284,8 @@ namespace UE::FractureEngine::Convex
 			bool bNoFailures = true;
 			for (int32 ConvexIdx : TransformToConvexInds[BoneIdx])
 			{
-				AppendConvexHullToCompactDynamicMesh(ConvexHulls[ConvexIdx].Get(), CombinedMesh, &GlobalTransformArray[BoneIdx]);
+				constexpr bool bConvertNonManifold = true; // Add non-manifold faces so we don't have holes messing up the sphere covering
+				AppendConvexHullToCompactDynamicMesh(ConvexHulls[ConvexIdx].Get(), CombinedMesh, &GlobalTransformArray[BoneIdx], bConvertNonManifold);
 			}
 			return bNoFailures;
 		};
diff --git a/Engine/Plugins/Experimental/Fracture/Source/FractureEngine/Public/FractureEngineConvex.h b/Engine/Plugins/Experimental/Fracture/Source/FractureEngine/Public/FractureEngineConvex.h
--- a/Engine/Plugins/Experimental/Fracture/Source/FractureEngine/Public/FractureEngineConvex.h
+++ b/Engine/Plugins/Experimental/Fracture/Source/FractureEngine/Public/FractureEngineConvex.h
@@ -4,6 +4,7 @@
 
 #include "Containers/ArrayView.h"
 #include "CompGeom/ConvexDecomposition3.h"
+#include "DynamicMesh/DynamicMesh3.h"
 
 namespace Chaos { class FConvex; }
 struct FManagedArrayCollection;
@@ -22,6 +23,12 @@ namespace UE::FractureEngine::Convex
 		bool bUseExistingVertexPositions = true;
 	};
 
+	/**
+	 * Get the convex hulls on the given Collection as a Dynamic Mesh. Optionally only include the hulls on the transforms in TransformSelection.
+	 * @return true if the collection has convex hulls, the bone selection (if present) was valid, and all requested convex hulls were included in the mesh
+	 */
+	bool FRACTUREENGINE_API GetConvexHullsAsDynamicMesh(const FManagedArrayCollection& Collection, UE::Geometry::FDynamicMesh3& UpdateMesh, bool bRestrictToSelection = false, const TArrayView<const int32> TransformSelection = TArrayView<const int32>());
+
 	/**
 	 * Simplify the convex hulls on the given Collection. Optionally only simplify the hulls on the transforms in TransformSelection.
 	 * @return true if the collection has convex hulls and the hulls were either simplified or did not need to be
diff --git a/Engine/Plugins/Experimental/GeometryCollectionPlugin/Source/GeometryCollectionNodes/Private/Dataflow/GeometryCollectionDebugNodes.cpp b/Engine/Plugins/Experimental/GeometryCollectionPlugin/Source/GeometryCollectionNodes/Private/Dataflow/GeometryCollectionDebugNodes.cpp
new file mode 100644
--- /dev/null
+++ b/Engine/Plugins/Experimental/GeometryCollectionPlugin/Source/GeometryCollectionNodes/Private/Dataflow/GeometryCollectionDebugNodes.cpp
@@ -0,0 +1,134 @@
+// Copyright Epic Games, Inc. All Rights Reserved.
+
+#include "Dataflow/GeometryCollectionDebugNodes.h"
+#include "Dataflow/DataflowCore.h"
+
+#include "Containers/UnrealString.h"
+
+#include "Misc/FileHelper.h"
+
+#include "DynamicMesh/DynamicMesh3.h"
+#include "DynamicMeshEditor.h"
+#include "Generators/BoxSphereGenerator.h"
+#include "FractureEngineConvex.h"
+
+#include UE_INLINE_GENERATED_CPP_BY_NAME(GeometryCollectionDebugNodes)
+
+namespace Dataflow
+{
+	void GeometryCollectionDebugNodes()
+	{
+		static const FLinearColor CDefaultNodeBodyTintColor = FLinearColor(0.f, 0.f, 0.f, 0.5f);
+
+		DATAFLOW_NODE_REGISTER_CREATION_FACTORY(FConvexHullToMeshDataflowNode);
+		DATAFLOW_NODE_REGISTER_CREATION_FACTORY(FSphereCoveringToMeshDataflowNode);
+		DATAFLOW_NODE_REGISTER_CREATION_FACTORY(FMeshToOBJStringDebugDataflowNode);
+		DATAFLOW_NODE_REGISTER_CREATION_FACTORY(FWriteStringToFile);
+
+		DATAFLOW_NODE_REGISTER_CREATION_FACTORY_NODE_COLORS_BY_CATEGORY("GeometryCollection|Development", FLinearColor(1.f, 0.f, 0.f), CDefaultNodeBodyTintColor);
+		DATAFLOW_NODE_REGISTER_CREATION_FACTORY_NODE_COLORS_BY_CATEGORY("GeometryCollection|Mesh", FLinearColor(1.f, 0.16f, 0.05f), CDefaultNodeBodyTintColor);
+	}
+}
+
+void FConvexHullToMeshDataflowNode::Evaluate(Dataflow::FContext& Context, const FDataflowOutput* Out) const
+{
+	using namespace UE::Geometry;
+
+	if (Out->IsA(&Mesh))
+	{
+		const FManagedArrayCollection& InCollection = GetValue<FManagedArrayCollection>(Context, &Collection);
+
+		TArray<int32> TransformSelection;
+		bool bHasSelection = IsConnected(&OptionalSelectionFilter);
+		if (bHasSelection)
+		{
+			const FDataflowTransformSelection& InOptionalSelectionFilter = GetValue(Context, &OptionalSelectionFilter);
+			TransformSelection = InOptionalSelectionFilter.AsArray();
+		}
+
+		FDynamicMesh3 HullsMesh;
+		UE::FractureEngine::Convex::GetConvexHullsAsDynamicMesh(InCollection, HullsMesh, bHasSelection, TransformSelection);
+
+		TObjectPtr<UDynamicMesh> NewMesh = NewObject<UDynamicMesh>();
+		NewMesh->SetMesh(MoveTemp(HullsMesh));
+		SetValue(Context, NewMesh, &Mesh);
+	}
+}
+
+void FSphereCoveringToMeshDataflowNode::Evaluate(Dataflow::FContext& Context, const FDataflowOutput* Out) const
+{
+	using namespace UE::Geometry;
+
+	if (Out->IsA(&Mesh))
+	{
+		TObjectPtr<UDynamicMesh> NewMesh = NewObject<UDynamicMesh>();
+		NewMesh->Reset();
+
+		FDynamicMesh3 Accum;
+		FDynamicMeshEditor Editor(&Accum);
+		FMeshIndexMappings IndexMaps_Unused;
+		FDataflowSphereCovering InSphereCovering = GetValue(Context, &SphereCovering);
+		for (int32 SphereIdx = 0; SphereIdx < InSphereCovering.Spheres.Num(); ++SphereIdx)
+		{
+			FBoxSphereGenerator SphereGen;
+			SphereGen.EdgeVertices.A = SphereGen.EdgeVertices.B = SphereGen.EdgeVertices.C = FMath::Max(2, VerticesAlongEachSide);
+			SphereGen.Radius = InSphereCovering.Spheres.GetRadius(SphereIdx);
+			FVector Center = InSphereCovering.Spheres.GetCenter(SphereIdx);
+			FDynamicMesh3 Sphere(&SphereGen.Generate());
+			Editor.AppendMesh(&Sphere, IndexMaps_Unused,
+				[Center](int VID, const FVector3d& Pos) { return Pos + Center; });
+			IndexMaps_Unused.Reset();
+		}
+		NewMesh->SetMesh(MoveTemp(Accum));
+
+		SetValue(Context, NewMesh, &Mesh);
+
+	}
+}
+
+void FMeshToOBJStringDebugDataflowNode::Evaluate(Dataflow::FContext& Context, const FDataflowOutput* Out) const
+{
+	if (Out->IsA(&StringOBJ))
+	{
+		const UDynamicMesh* InMesh = GetValue(Context, &Mesh);
+		const UE::Geometry::FDynamicMesh3& MeshRef = InMesh->GetMeshRef();
+		FString Build;
+		Build.Reset(MeshRef.MaxVertexID() * 40 + MeshRef.TriangleCount() * 24);
+		for (int32 VID = 0; VID < MeshRef.MaxVertexID(); ++VID)
+		{
+			FVector V = FVector::ZeroVector;
+			if (MeshRef.IsVertex(VID))
+			{
+				V = MeshRef.GetVertex(VID);
+			}
+			Build += FString::Printf(TEXT("v %f %f %f\n"), V.X, V.Y, V.Z);
+		}
+		bool bInInvertFaces = GetValue(Context, &bInvertFaces);
+		for (int32 TID = 0; TID < MeshRef.MaxTriangleID(); ++TID)
+		{
+			if (MeshRef.IsTriangle(TID))
+			{
+				UE::Geometry::FIndex3i T = MeshRef.GetTriangle(TID);
+				// Note: OBJ viewers generally expect the opposite triangle winding from UE meshes, so inverted == the UE order
+				if (!bInInvertFaces)
+				{
+					Swap(T.B, T.C);
+				}
+				Build += FString::Printf(TEXT("f %d %d %d\n"), T.A + 1, T.B + 1, T.C + 1);
+			}
+		}
+		SetValue(Context, MoveTemp(Build), &StringOBJ);
+	}
+}
+
+void FWriteStringToFile::Evaluate(Dataflow::FContext& Context, const FDataflowOutput* Out) const
+{
+	FString InContents = GetValue(Context, &FileContents);
+	FString InPath = GetValue(Context, &FilePath);
+
+	bool bSuccess = FFileHelper::SaveStringToFile(InContents, *InPath);
+	if (!bSuccess)
+	{
+		UE_LOG(LogChaos, Warning, TEXT("Failed to write to file %s:\n\n%s"), *InPath, *InContents);
+	}
+}
diff --git a/Engine/Plugins/Experimental/GeometryCollectionPlugin/Source/GeometryCollectionNodes/Private/Dataflow/GeometryCollectionNodesPlugin.cpp b/Engine/Plugins/Experimental/GeometryCollectionPlugin/Source/GeometryCollectionNodes/Private/Dataflow/GeometryCollectionNodesPlugin.cpp
--- a/Engine/Plugins/Experimental/GeometryCollectionPlugin/Source/GeometryCollectionNodes/Private/Dataflow/GeometryCollectionNodesPlugin.cpp
+++ b/Engine/Plugins/Experimental/GeometryCollectionPlugin/Source/GeometryCollectionNodes/Private/Dataflow/GeometryCollectionNodesPlugin.cpp
@@ -23,6 +23,7 @@
 #include "Dataflow/GeometryCollectionConversionNodes.h"
 #include "Dataflow/GeometryCollectionVerticesNodes.h"
 #include "Dataflow/GeometryCollectionArrayNodes.h"
+#include "Dataflow/GeometryCollectionDebugNodes.h"
 
 #define LOCTEXT_NAMESPACE "DataflowNodes"
 
@@ -47,6 +48,7 @@ void IGeometryCollectionNodesPlugin::StartupModule()
 	Dataflow::GeometryCollectionConversionNodes();
 	Dataflow::GeometryCollectionVerticesNodes();
 	Dataflow::GeometryCollectionArrayNodes();
+	Dataflow::GeometryCollectionDebugNodes();
 }
 
 void IGeometryCollectionNodesPlugin::ShutdownModule()
diff --git a/Engine/Plugins/Experimental/GeometryCollectionPlugin/Source/GeometryCollectionNodes/Private/Dataflow/GeometryCollectionUtilityNodes.cpp b/Engine/Plugins/Experimental/GeometryCollectionPlugin/Source/GeometryCollectionNodes/Private/Dataflow/GeometryCollectionUtilityNodes.cpp
--- a/Engine/Plugins/Experimental/GeometryCollectionPlugin/Source/GeometryCollectionNodes/Private/Dataflow/GeometryCollectionUtilityNodes.cpp
+++ b/Engine/Plugins/Experimental/GeometryCollectionPlugin/Source/GeometryCollectionNodes/Private/Dataflow/GeometryCollectionUtilityNodes.cpp
@@ -202,13 +202,15 @@ FGenerateClusterConvexHullsFromLeafHullsDataflowNode::FGenerateClusterConvexHull
 	RegisterInputConnection(&MinRadius);
 
 	RegisterOutputConnection(&Collection);
+	RegisterOutputConnection(&SphereCovering);
 }
 
 void FGenerateClusterConvexHullsFromLeafHullsDataflowNode::Evaluate(Dataflow::FContext& Context, const FDataflowOutput* Out) const
 {
-	if (Out->IsA<FManagedArrayCollection>(&Collection) && IsConnected(&Collection))
+	if (Out->IsA(&Collection) || Out->IsA(&SphereCovering))
 	{
 		const FManagedArrayCollection& InCollection = GetValue<FManagedArrayCollection>(Context, &Collection);
+		FDataflowSphereCovering Spheres;
 
 		if (TUniquePtr<FGeometryCollection> GeomCollection = TUniquePtr<FGeometryCollection>(InCollection.NewCopy<FGeometryCollection>()))
 		{
@@ -254,8 +256,17 @@ void FGenerateClusterConvexHullsFromLeafHullsDataflowNode::Evaluate(Dataflow::FC
 				);
 			}
 
-			SetValue<FManagedArrayCollection>(Context, static_cast<const FManagedArrayCollection>(*GeomCollection), &Collection);
+			SetValue(Context, static_cast<const FManagedArrayCollection>(*GeomCollection), &Collection);
+			// Move the negative space to the output container at the end to be sure it is no longer needed
+			Spheres.Spheres = MoveTemp(NegativeSpace);
+		}
+		else
+		{
+			UE_LOG(LogChaos, Error, TEXT("Error: Input collection could not be converted to a valid Geometry Collection"));
+			SetValue(Context, InCollection, &Collection);
 		}
+
+		SetValue(Context, MoveTemp(Spheres), &SphereCovering);
 	}
 }
 
@@ -273,13 +284,15 @@ FGenerateClusterConvexHullsFromChildrenHullsDataflowNode::FGenerateClusterConvex
 	RegisterInputConnection(&MinRadius);
 
 	RegisterOutputConnection(&Collection);
+	RegisterOutputConnection(&SphereCovering);
 }
 
 void FGenerateClusterConvexHullsFromChildrenHullsDataflowNode::Evaluate(Dataflow::FContext& Context, const FDataflowOutput* Out) const
 {
-	if (Out->IsA<FManagedArrayCollection>(&Collection) && IsConnected(&Collection))
+	if (Out->IsA(&Collection) || Out->IsA(&SphereCovering))
 	{
 		const FManagedArrayCollection& InCollection = GetValue<FManagedArrayCollection>(Context, &Collection);
+		FDataflowSphereCovering Spheres;
 
 		if (TUniquePtr<FGeometryCollection> GeomCollection = TUniquePtr<FGeometryCollection>(InCollection.NewCopy<FGeometryCollection>()))
 		{
@@ -326,7 +339,16 @@ void FGenerateClusterConvexHullsFromChildrenHullsDataflowNode::Evaluate(Dataflow
 			}
 
 			SetValue<FManagedArrayCollection>(Context, static_cast<const FManagedArrayCollection>(*GeomCollection), &Collection);
+			// Move the negative space to the output container at the end to be sure it is no longer needed
+			Spheres.Spheres = MoveTemp(NegativeSpace);
+		}
+		else
+		{
+			UE_LOG(LogChaos, Error, TEXT("Error: Input collection could not be converted to a valid Geometry Collection"));
+			SetValue(Context, InCollection, &Collection);
 		}
+		
+		SetValue(Context, MoveTemp(Spheres), &SphereCovering);
 	}
 }
 
diff --git a/Engine/Plugins/Experimental/GeometryCollectionPlugin/Source/GeometryCollectionNodes/Public/Dataflow/GeometryCollectionDebugNodes.h b/Engine/Plugins/Experimental/GeometryCollectionPlugin/Source/GeometryCollectionNodes/Public/Dataflow/GeometryCollectionDebugNodes.h
new file mode 100644
--- /dev/null
+++ b/Engine/Plugins/Experimental/GeometryCollectionPlugin/Source/GeometryCollectionNodes/Public/Dataflow/GeometryCollectionDebugNodes.h
@@ -0,0 +1,135 @@
+// Copyright Epic Games, Inc. All Rights Reserved.
+
+#pragma once 
+
+#include "CoreMinimal.h"
+#include "Dataflow/DataflowEngine.h"
+#include "GeometryCollection/ManagedArrayCollection.h"
+#include "DynamicMesh/DynamicMesh3.h"
+#include "UDynamicMesh.h"
+#include "GeometryCollectionUtilityNodes.h"
+
+#include "GeometryCollectionDebugNodes.generated.h"
+
+//~ Dataflow nodes that are useful for debugging a geometry collection-related dataflow graph
+
+
+// Convert convex hulls on a geometry collection to a dynamic mesh
+USTRUCT(meta = (DataflowGeometryCollection))
+struct FConvexHullToMeshDataflowNode : public FDataflowNode
+{
+	GENERATED_USTRUCT_BODY()
+	DATAFLOW_NODE_DEFINE_INTERNAL(FConvexHullToMeshDataflowNode, "Convex Hull to Mesh", "GeometryCollection|Mesh", "")
+
+public:
+	UPROPERTY(meta = (DataflowInput, DataflowIntrinsic))
+	FManagedArrayCollection Collection;
+
+	/** Optional transform selection to convert hulls from -- if not provided, all convex hulls will be converted. */
+	UPROPERTY(meta = (DataflowInput, DataflowIntrinsic))
+	FDataflowTransformSelection OptionalSelectionFilter;
+
+	UPROPERTY(meta = (DataflowOutput))
+	TObjectPtr<UDynamicMesh> Mesh;
+
+	FConvexHullToMeshDataflowNode(const Dataflow::FNodeParameters& InParam, FGuid InGuid = FGuid::NewGuid())
+		: FDataflowNode(InParam, InGuid)
+	{
+		RegisterInputConnection(&Collection);
+		RegisterInputConnection(&OptionalSelectionFilter);
+		RegisterOutputConnection(&Mesh);
+	}
+
+	virtual void Evaluate(Dataflow::FContext& Context, const FDataflowOutput* Out) const override;
+};
+
+// Convert a sphere covering, as generated by the 'protect negative space' option on the "Generate Cluster Convex Hull" nodes, to a dynamic mesh
+USTRUCT(meta = (DataflowGeometryCollection))
+struct FSphereCoveringToMeshDataflowNode : public FDataflowNode
+{
+	GENERATED_USTRUCT_BODY()
+	DATAFLOW_NODE_DEFINE_INTERNAL(FSphereCoveringToMeshDataflowNode, "Sphere Covering to Mesh", "GeometryCollection|Mesh", "")
+
+public:
+	// The sphere covering to convert to a mesh
+	UPROPERTY(meta = (DataflowInput, DataflowIntrinsic))
+	FDataflowSphereCovering SphereCovering;
+
+	// Number of vertices to use along each axis when creating a mesh for each sphere
+	UPROPERTY(EditAnywhere, Category = Options, meta = (ClampMin = 2))
+	int32 VerticesAlongEachSide = 8;
+
+	UPROPERTY(meta = (DataflowOutput))
+	TObjectPtr<UDynamicMesh> Mesh;
+
+	FSphereCoveringToMeshDataflowNode(const Dataflow::FNodeParameters& InParam, FGuid InGuid = FGuid::NewGuid())
+		: FDataflowNode(InParam, InGuid)
+	{
+		RegisterInputConnection(&SphereCovering);
+		RegisterOutputConnection(&Mesh);
+	}
+
+	virtual void Evaluate(Dataflow::FContext& Context, const FDataflowOutput* Out) const override;
+};
+
+// Convert a mesh to a string formatted as an OBJ file, which can be read by many external mesh viewers
+USTRUCT(meta = (DataflowGeometryCollection))
+struct FMeshToOBJStringDebugDataflowNode: public FDataflowNode
+{
+	GENERATED_USTRUCT_BODY()
+	DATAFLOW_NODE_DEFINE_INTERNAL(FMeshToOBJStringDebugDataflowNode, "Convert Mesh to OBJ String", "GeometryCollection|Development", "")
+
+public:
+	UPROPERTY(meta = (DataflowInput))
+	TObjectPtr<UDynamicMesh> Mesh;
+
+	// Whether to flip the orientation of the triangles in the OBJ output
+	UPROPERTY(EditAnywhere, Category = Options, meta = (DataflowInput))
+	bool bInvertFaces;
+
+	// A string representing the input mesh in the OBJ file format
+	UPROPERTY(meta = (DataflowOutput))
+	FString StringOBJ;
+
+	FMeshToOBJStringDebugDataflowNode(const Dataflow::FNodeParameters& InParam, FGuid InGuid = FGuid::NewGuid())
+		: FDataflowNode(InParam, InGuid)
+	{
+		RegisterInputConnection(&Mesh);
+		RegisterInputConnection(&bInvertFaces);
+		RegisterOutputConnection(&StringOBJ);
+	}
+
+	virtual void Evaluate(Dataflow::FContext& Context, const FDataflowOutput* Out) const override;
+};
+
+// Write a string to a file
+USTRUCT(meta = (DataflowGeometryCollection))
+struct FWriteStringToFile : public FDataflowNode
+{
+	GENERATED_USTRUCT_BODY()
+	DATAFLOW_NODE_DEFINE_INTERNAL(FWriteStringToFile, "Write String to File", "GeometryCollection|Development", "")
+
+public:
+
+	// Where file should be written on disk
+	UPROPERTY(EditAnywhere, Category = Options, meta = (DataflowInput))
+	FString FilePath;
+
+	// Contents of the file to write
+	UPROPERTY(EditAnywhere, Category = Options, meta = (DataflowInput))
+	FString FileContents;
+
+	FWriteStringToFile(const Dataflow::FNodeParameters& InParam, FGuid InGuid = FGuid::NewGuid())
+		: FDataflowNode(InParam, InGuid)
+	{
+		RegisterInputConnection(&FileContents);
+		RegisterInputConnection(&FilePath);
+	}
+
+	virtual void Evaluate(Dataflow::FContext& Context, const FDataflowOutput* Out) const override;
+};
+
+namespace Dataflow
+{
+	void GeometryCollectionDebugNodes();
+}
diff --git a/Engine/Plugins/Experimental/GeometryCollectionPlugin/Source/GeometryCollectionNodes/Public/Dataflow/GeometryCollectionUtilityNodes.h b/Engine/Plugins/Experimental/GeometryCollectionPlugin/Source/GeometryCollectionNodes/Public/Dataflow/GeometryCollectionUtilityNodes.h
--- a/Engine/Plugins/Experimental/GeometryCollectionPlugin/Source/GeometryCollectionNodes/Public/Dataflow/GeometryCollectionUtilityNodes.h
+++ b/Engine/Plugins/Experimental/GeometryCollectionPlugin/Source/GeometryCollectionNodes/Public/Dataflow/GeometryCollectionUtilityNodes.h
@@ -7,6 +7,7 @@
 #include "Dataflow/DataflowSelection.h"
 #include "GeometryCollection/ManagedArrayCollection.h"
 #include "GeometryCollection/GeometryCollectionConvexUtility.h"
+#include "CompGeom/ConvexDecomposition3.h"
 
 #include "GeometryCollectionUtilityNodes.generated.h"
 
@@ -210,6 +211,17 @@ public:
 
 };
 
+//~ Simple wrapper class to make the sphere covering data possible to pass via dataflow
+// A set of spheres generated to represent empty space when creating a minimal set of convex hulls, e.g. in one of the Generate Cluster Convex Hulls nodes
+USTRUCT()
+struct FDataflowSphereCovering
+{
+	GENERATED_USTRUCT_BODY()
+
+public:
+	UE::Geometry::FSphereCovering Spheres;
+};
+
 /**
  *
  * Generates cluster convex hulls for leafs hulls
@@ -225,6 +237,10 @@ public:
 	UPROPERTY(meta = (DataflowInput, DataflowOutput))
 	FManagedArrayCollection Collection;
 
+	// A representation of the negative space protected by the 'protect negative space' option. If negative space is not protected, this will contain zero spheres.
+	UPROPERTY(meta = (DataflowOutput))
+	FDataflowSphereCovering SphereCovering;
+
 	/** Maximum number of convex to generate for a specific cluster. Will be ignored if error tolerance is used instead */
 	UPROPERTY(EditAnywhere, Category = "Convex", meta = (DataflowInput, EditCondition = "ErrorTolerance == 0"))
 	int32 ConvexCount = 2;
@@ -289,6 +305,10 @@ struct FGenerateClusterConvexHullsFromChildrenHullsDataflowNode : public FDatafl
 public:
 	UPROPERTY(meta = (DataflowInput, DataflowOutput))
 	FManagedArrayCollection Collection;
+	
+	// A representation of the negative space protected by the 'protect negative space' option. If negative space is not protected, this will contain zero spheres.
+	UPROPERTY(meta = (DataflowOutput))
+	FDataflowSphereCovering SphereCovering;
 
 	/** Maximum number of convex to generate for a specific cluster. Will be ignored if error tolerance is used instead */
 	UPROPERTY(EditAnywhere, Category = "Convex", meta = (DataflowInput, EditCondition = "ErrorTolerance == 0"))
