Add additional convex simplification option based on an angle threshold; can give more intuitive results in some cases

diff --git a/Engine/Plugins/Experimental/Fracture/Source/FractureEngine/FractureEngine.Build.cs b/Engine/Plugins/Experimental/Fracture/Source/FractureEngine/FractureEngine.Build.cs
--- a/Engine/Plugins/Experimental/Fracture/Source/FractureEngine/FractureEngine.Build.cs
+++ b/Engine/Plugins/Experimental/Fracture/Source/FractureEngine/FractureEngine.Build.cs
@@ -11,11 +11,13 @@ namespace UnrealBuildTool.Rules
 				{
 					"Chaos",
 					"Core",
+					"CoreUObject",
 					"DataflowCore",
 					"PlanarCut",
 					"GeometryCore",
 					"DynamicMesh",
-					"Voronoi"
+					"Voronoi",
+					"MeshUtilitiesCommon", // for FDisjointSet
 				}
 				);
 		}
diff --git a/Engine/Plugins/Experimental/Fracture/Source/FractureEngine/Private/FractureEngineConvex.cpp b/Engine/Plugins/Experimental/Fracture/Source/FractureEngine/Private/FractureEngineConvex.cpp
--- a/Engine/Plugins/Experimental/Fracture/Source/FractureEngine/Private/FractureEngineConvex.cpp
+++ b/Engine/Plugins/Experimental/Fracture/Source/FractureEngine/Private/FractureEngineConvex.cpp
@@ -11,6 +11,9 @@
 #include "GeometryCollection/GeometryCollectionConvexUtility.h"
 #include "GeometryCollection/Facades/CollectionTransformFacade.h"
 #include "GeometryCollection/Facades/CollectionTransformSelectionFacade.h"
+#include "DisjointSet.h"
+#include "VectorUtil.h"
+#include "Util/IndexPriorityQueue.h"
 
 namespace
 {
@@ -200,49 +203,172 @@ namespace UE::FractureEngine::Convex
 			}
 			return true;
 		}
-	
-		// Convert to DynamicMesh to run simplifier
-		UE::Geometry::FDynamicMesh3 Mesh = ConvexHullToDynamicMesh(InConvexHull);
-		
-		// Run simplification
-		UE::Geometry::FVolPresMeshSimplification Simplifier(&Mesh);
-		
-		Simplifier.CollapseMode =
-			Settings.bUseExistingVertexPositions ?
-				UE::Geometry::FVolPresMeshSimplification::ESimplificationCollapseModes::MinimalExistingVertexError
-			:	UE::Geometry::FVolPresMeshSimplification::ESimplificationCollapseModes::MinimalQuadricPositionError;
-		if (Settings.bUseGeometricTolerance)
-		{
-			Simplifier.GeometricErrorConstraint = UE::Geometry::FVolPresMeshSimplification::EGeometricErrorCriteria::PredictedPointToProjectionTarget;
-			Simplifier.GeometricErrorTolerance = Settings.ErrorTolerance;
-		}
-		if (Settings.bUseGeometricTolerance)
-		{
-			// Simplify to the smallest non-degenerate number of triangles, relying on geometric error criteria
-			UE::Geometry::FDynamicMesh3 ProjectionTargetMesh(Mesh);
-			UE::Geometry::FDynamicMeshAABBTree3 ProjectionTargetSpatial(&ProjectionTargetMesh, true);
-			UE::Geometry::FMeshProjectionTarget ProjTarget(&ProjectionTargetMesh, &ProjectionTargetSpatial);
-			Simplifier.SetProjectionTarget(&ProjTarget);
-			int32 TargetTriCount = Settings.bUseTargetTriangleCount ? Settings.TargetTriangleCount : 4;
-			Simplifier.SimplifyToTriangleCount(TargetTriCount);
-		}
-		else if (Settings.bUseTargetTriangleCount)
+
+		if (Settings.SimplifyMethod == EConvexHullSimplifyMethod::MeshQSlim)
 		{
-			Simplifier.SimplifyToTriangleCount(Settings.TargetTriangleCount);
+			// Convert to DynamicMesh to run simplifier
+			UE::Geometry::FDynamicMesh3 Mesh = ConvexHullToDynamicMesh(InConvexHull);
+
+			// Run simplification
+			UE::Geometry::FVolPresMeshSimplification Simplifier(&Mesh);
+
+			Simplifier.CollapseMode =
+				Settings.bUseExistingVertexPositions ?
+					UE::Geometry::FVolPresMeshSimplification::ESimplificationCollapseModes::MinimalExistingVertexError
+				: 	UE::Geometry::FVolPresMeshSimplification::ESimplificationCollapseModes::MinimalQuadricPositionError;
+			if (Settings.bUseGeometricTolerance)
+			{
+				Simplifier.GeometricErrorConstraint = UE::Geometry::FVolPresMeshSimplification::EGeometricErrorCriteria::PredictedPointToProjectionTarget;
+				Simplifier.GeometricErrorTolerance = Settings.ErrorTolerance;
+			}
+			if (Settings.bUseGeometricTolerance)
+			{
+				// Simplify to the smallest non-degenerate number of triangles, relying on geometric error criteria
+				UE::Geometry::FDynamicMesh3 ProjectionTargetMesh(Mesh);
+				UE::Geometry::FDynamicMeshAABBTree3 ProjectionTargetSpatial(&ProjectionTargetMesh, true);
+				UE::Geometry::FMeshProjectionTarget ProjTarget(&ProjectionTargetMesh, &ProjectionTargetSpatial);
+				Simplifier.SetProjectionTarget(&ProjTarget);
+				int32 TargetTriCount = Settings.bUseTargetTriangleCount ? Settings.TargetTriangleCount : 4;
+				Simplifier.SimplifyToTriangleCount(TargetTriCount);
+			}
+			else if (Settings.bUseTargetTriangleCount)
+			{
+				Simplifier.SimplifyToTriangleCount(Settings.TargetTriangleCount);
+			}
+			else
+			{
+				// Note: Quadric error threshold doesn't have the same geometric meaning as distance; this is not equivalent to using a geometric error tolerance
+				Simplifier.SimplifyToMaxError(Settings.ErrorTolerance * Settings.ErrorTolerance);
+			}
+
+			TArray<::Chaos::FVec3f> NewConvexVerts;
+			NewConvexVerts.Reserve(Mesh.VertexCount());
+			for (int32 VIdx : Mesh.VertexIndicesItr())
+			{
+				NewConvexVerts.Add((::Chaos::FVec3f)Mesh.GetVertex(VIdx));
+			}
+			*OutConvexHull = ::Chaos::FConvex(NewConvexVerts, InConvexHull->GetMargin(), ::Chaos::FConvexBuilder::EBuildMethod::Default);
 		}
-		else
+		else if (Settings.SimplifyMethod == EConvexHullSimplifyMethod::AngleTolerance)
 		{
-			// Note: Quadric error threshold doesn't have the same geometric meaning as distance; this is not equivalent to using a geometric error tolerance
-			Simplifier.SimplifyToMaxError(Settings.ErrorTolerance * Settings.ErrorTolerance);
-		}
+			FDisjointSet PlaneGroups(NumP);
+			TArray<FVector> PlaneNormals; PlaneNormals.SetNumUninitialized(NumP);
+			TArray<double> PlaneAreas; PlaneAreas.SetNumUninitialized(NumP);
+			for (int32 PlaneIdx = 0; PlaneIdx < NumP; ++PlaneIdx)
+			{
+				Chaos::FVec3 Pos, Normal;
+				InConvexHull->GetPlaneNX(PlaneIdx, Normal, Pos);
+				PlaneNormals[PlaneIdx] = (FVector)Normal;
+				Chaos::FVec3 V0 = InConvexHull->GetVertex(InConvexHull->GetPlaneVertex(PlaneIdx, 0));
+				double AreaSum = 0;
+				for (int32 SubIdx = 1; SubIdx + 1 < InConvexHull->NumPlaneVertices(PlaneIdx); ++SubIdx)
+				{
+					Chaos::FVec3 V1 = InConvexHull->GetVertex(InConvexHull->GetPlaneVertex(PlaneIdx, SubIdx));
+					Chaos::FVec3 V2 = InConvexHull->GetVertex(InConvexHull->GetPlaneVertex(PlaneIdx, SubIdx+1));
+					AreaSum += UE::Geometry::VectorUtil::Area(V0, V1, V2);
+				}
+				PlaneAreas[PlaneIdx] = AreaSum;
+			}
 
-		TArray<::Chaos::FVec3f> NewConvexVerts;
-		NewConvexVerts.Reserve(Mesh.VertexCount());
-		for (int32 VIdx : Mesh.VertexIndicesItr())
-		{
-			NewConvexVerts.Add((::Chaos::FVec3f)Mesh.GetVertex(VIdx));
+			int32 NumEdges = InConvexHull->NumEdges();
+			UE::Geometry::FIndexPriorityQueue EdgeQueue(NumEdges);
+
+			const double AreaThreshold = FMath::Max(UE_DOUBLE_KINDA_SMALL_NUMBER, Settings.SmallAreaThreshold);
+			auto GetMergeWeight = [&InConvexHull, &PlaneGroups, &PlaneAreas, &PlaneNormals, AreaThreshold](int32 EdgeIdx, bool bHasGroups) -> float
+			{
+				int32 P[2]{ InConvexHull->GetEdgePlane(EdgeIdx, 0), InConvexHull->GetEdgePlane(EdgeIdx, 1) };
+				if (bHasGroups)
+				{
+					P[0] = PlaneGroups.Find(P[0]);
+					P[1] = PlaneGroups.Find(P[1]);
+					if (P[0] == P[1])
+					{
+						// return a value higher than the max possible angle threshold if the plane groups are already merged
+						constexpr float CannotMergeValue = 4;
+						return CannotMergeValue;
+					}
+				}
+				// Planes with small area can be merged into any neighbor (e.g., to filter degenerate or less meaningful normals)
+				if (PlaneAreas[P[0]] < AreaThreshold || PlaneAreas[P[1]] < AreaThreshold)
+				{
+					return 0;
+				}
+				float NormalAlignment = 1 - float(PlaneNormals[P[0]].Dot(PlaneNormals[P[1]]));
+				return NormalAlignment;
+			};
+			
+			const float Threshold = 1 - FMath::Cos(FMath::DegreesToRadians(Settings.AngleThreshold));
+			for (int32 EdgeIdx = 0; EdgeIdx < NumEdges; ++EdgeIdx)
+			{
+				float Wt = GetMergeWeight(EdgeIdx, false);
+				if (Wt < Threshold)
+				{
+					EdgeQueue.Insert(EdgeIdx, Wt);
+				}
+			}
+
+			int32 NumRemainingPlanes = NumP;
+			while (EdgeQueue.GetCount() > 0)
+			{
+				int32 EdgeIdx = EdgeQueue.Dequeue();
+				float Wt = GetMergeWeight(EdgeIdx, true);
+				if (Wt < Threshold)
+				{
+					int32 P[2]{ PlaneGroups.Find(InConvexHull->GetEdgePlane(EdgeIdx, 0)), PlaneGroups.Find(InConvexHull->GetEdgePlane(EdgeIdx, 1)) };
+					PlaneGroups.Union(P[0], P[1]);
+					int32 Parent = PlaneGroups.Find(P[0]);
+					FVector Normal = PlaneNormals[P[0]] * PlaneAreas[P[0]] + PlaneNormals[P[1]] * PlaneAreas[P[1]];
+					Normal.Normalize();
+					double AreaSum = PlaneAreas[P[0]] + PlaneAreas[P[1]];
+					PlaneNormals[Parent] = Normal;
+					PlaneAreas[Parent] = AreaSum;
+					NumRemainingPlanes--;
+				}
+				if (NumRemainingPlanes < Settings.TargetTriangleCount)
+				{
+					break;
+				}
+			}
+
+			TArray<::Chaos::FVec3f> NewConvexVerts;
+			int32 NumV = InConvexHull->NumVertices();
+			NewConvexVerts.Reserve(NumV);
+
+			// Note: We should be able to do this more directly by using ConvexStructure.FindVertexPlanes,
+			// but currently this appears to sometimes not find all the vertex planes, so over-simplifies.
+			TArray<FIntVector> VertPlaneGroups;
+			VertPlaneGroups.Init(FIntVector(-1, -1, -1), NumV);
+			for (int32 PIdx = 0; PIdx < NumP; ++PIdx)
+			{
+				int32 Group = PlaneGroups.Find(PIdx);
+				int32 NumFaceVert = InConvexHull->NumPlaneVertices(PIdx);
+				for (int32 SubIdx = 0; SubIdx < NumFaceVert; ++SubIdx)
+				{
+					int32 VIdx = InConvexHull->GetPlaneVertex(PIdx, SubIdx);
+					for (int32 Idx = 0; Idx < 3; ++Idx)
+					{
+						if (VertPlaneGroups[VIdx][Idx] == -1)
+						{
+							VertPlaneGroups[VIdx][Idx] = Group;
+							break;
+						}
+						else if (VertPlaneGroups[VIdx][Idx] == Group)
+						{
+							break;
+						}
+					}
+				}
+			}
+			for (int32 VIdx = 0; VIdx < NumV; ++VIdx)
+			{
+				if (VertPlaneGroups[VIdx][2] > -1)
+				{
+					NewConvexVerts.Add(InConvexHull->GetVertex(VIdx));
+				}
+			}
+
+			*OutConvexHull = ::Chaos::FConvex(NewConvexVerts, InConvexHull->GetMargin(), ::Chaos::FConvexBuilder::EBuildMethod::Default);
 		}
-		*OutConvexHull = ::Chaos::FConvex(NewConvexVerts, InConvexHull->GetMargin(), ::Chaos::FConvexBuilder::EBuildMethod::Default);
 
 		return true;
 	}
diff --git a/Engine/Plugins/Experimental/Fracture/Source/FractureEngine/Public/FractureEngineConvex.h b/Engine/Plugins/Experimental/Fracture/Source/FractureEngine/Public/FractureEngineConvex.h
--- a/Engine/Plugins/Experimental/Fracture/Source/FractureEngine/Public/FractureEngineConvex.h
+++ b/Engine/Plugins/Experimental/Fracture/Source/FractureEngine/Public/FractureEngineConvex.h
@@ -9,11 +9,23 @@
 namespace Chaos { class FConvex; }
 struct FManagedArrayCollection;
 
+UENUM()
+enum class EConvexHullSimplifyMethod
+{
+	MeshQSlim,
+	AngleTolerance
+};
+
 namespace UE::FractureEngine::Convex
 {
 	// Options to control the simplification of an existing convex hull
 	struct FSimplifyHullSettings
 	{
+		EConvexHullSimplifyMethod SimplifyMethod = EConvexHullSimplifyMethod::MeshQSlim;
+
+		double AngleThreshold = 0; // used for AngleTolerance simplification
+		double SmallAreaThreshold = .1; // optionally used for AngleTolerance to ignore very small faces
+
 		bool bUseGeometricTolerance = true;
 		double ErrorTolerance = 5;
 
diff --git a/Engine/Plugins/Experimental/GeometryCollectionPlugin/Source/GeometryCollectionNodes/Private/Dataflow/GeometryCollectionUtilityNodes.cpp b/Engine/Plugins/Experimental/GeometryCollectionPlugin/Source/GeometryCollectionNodes/Private/Dataflow/GeometryCollectionUtilityNodes.cpp
--- a/Engine/Plugins/Experimental/GeometryCollectionPlugin/Source/GeometryCollectionNodes/Private/Dataflow/GeometryCollectionUtilityNodes.cpp
+++ b/Engine/Plugins/Experimental/GeometryCollectionPlugin/Source/GeometryCollectionNodes/Private/Dataflow/GeometryCollectionUtilityNodes.cpp
@@ -9,8 +9,6 @@
 #include "GeometryCollection/Facades/CollectionHierarchyFacade.h"
 #include "GeometryCollection/Facades/CollectionTransformSelectionFacade.h"
 
-#include "FractureEngineConvex.h"
-
 #include "MeshSimplification.h"
 
 
@@ -116,6 +114,7 @@ FSimplifyConvexHullsDataflowNode::FSimplifyConvexHullsDataflowNode(const Dataflo
 {
 	RegisterInputConnection(&Collection);
 	RegisterInputConnection(&OptionalSelectionFilter);
+	RegisterInputConnection(&SimplificationAngleThreshold);
 	RegisterInputConnection(&SimplificationDistanceThreshold);
 	RegisterInputConnection(&MinTargetTriangleCount);
 	RegisterOutputConnection(&Collection);
@@ -142,7 +141,9 @@ void FSimplifyConvexHullsDataflowNode::Evaluate(Dataflow::FContext& Context, con
 		}
 
 		UE::FractureEngine::Convex::FSimplifyHullSettings Settings;
+		Settings.SimplifyMethod = SimplifyMethod;
 		Settings.ErrorTolerance = GetValue(Context, &SimplificationDistanceThreshold);
+		Settings.AngleThreshold = GetValue(Context, &SimplificationAngleThreshold);
 		Settings.bUseGeometricTolerance = true;
 		Settings.bUseTargetTriangleCount = true;
 		Settings.bUseExistingVertexPositions = bUseExistingVertices;
diff --git a/Engine/Plugins/Experimental/GeometryCollectionPlugin/Source/GeometryCollectionNodes/Public/Dataflow/GeometryCollectionUtilityNodes.h b/Engine/Plugins/Experimental/GeometryCollectionPlugin/Source/GeometryCollectionNodes/Public/Dataflow/GeometryCollectionUtilityNodes.h
--- a/Engine/Plugins/Experimental/GeometryCollectionPlugin/Source/GeometryCollectionNodes/Public/Dataflow/GeometryCollectionUtilityNodes.h
+++ b/Engine/Plugins/Experimental/GeometryCollectionPlugin/Source/GeometryCollectionNodes/Public/Dataflow/GeometryCollectionUtilityNodes.h
@@ -9,6 +9,8 @@
 #include "GeometryCollection/GeometryCollectionConvexUtility.h"
 #include "CompGeom/ConvexDecomposition3.h"
 
+#include "FractureEngineConvex.h"
+
 #include "GeometryCollectionUtilityNodes.generated.h"
 
 
@@ -153,12 +155,19 @@ public:
 	UPROPERTY(meta = (DataflowInput, DataflowIntrinsic))
 	FDataflowTransformSelection OptionalSelectionFilter;
 
-	/** Computed convex hulls are simplified to keep points spaced at least this far apart (except where needed to keep the hull from collapsing to zero volume) */
-	UPROPERTY(EditAnywhere, Category = "Convex", meta = (DataflowInput, ClampMin = 0.f))
+	UPROPERTY(EditAnywhere, Category = "Convex")
+	EConvexHullSimplifyMethod SimplifyMethod = EConvexHullSimplifyMethod::MeshQSlim;
+
+	/** Simplified hull should preserve angles larger than this (in degrees).  Used by the AngleTolerance simplification method. */
+	UPROPERTY(EditAnywhere, Category = "Convex", meta = (DataflowInput, ClampMin = 0.f, EditCondition = "SimplifyMethod == EConvexHullSimplifyMethod::AngleTolerance"))
+	float SimplificationAngleThreshold = 10.f;
+
+	/** Simplified hull should stay within this distance of the initial convex hull. Used by the MeshQSlim simplification method. */
+	UPROPERTY(EditAnywhere, Category = "Convex", meta = (DataflowInput, ClampMin = 0.f, EditCondition = "SimplifyMethod == EConvexHullSimplifyMethod::MeshQSlim"))
 	float SimplificationDistanceThreshold = 10.f;
 
-	/** The minimum number of triangles to use for the convex hull. Will use more triangles than this as-needed to preserve the distance threshold. */
-	UPROPERTY(EditAnywhere, Category = "Convex", meta = (DataflowInput, ClampMin = 4))
+	/** The minimum number of faces to use for the convex hull. For MeshQSlim simplification, this is a triangle count, which may be further reduced on conversion back to a convex hull. */
+	UPROPERTY(EditAnywhere, Category = "Convex", meta = (DisplayName = "Min Target Face Count", DataflowInput, ClampMin = 4))
 	int32 MinTargetTriangleCount = 12;
 
 	/** Whether to restrict the simplified hulls to only use vertices from the original hulls. */
