fix non-unity build failure on missing includes for UENUM in FractureEngineConvex.h

diff --git a/Engine/Plugins/Experimental/Fracture/Source/FractureEngine/Public/FractureEngineConvex.h b/Engine/Plugins/Experimental/Fracture/Source/FractureEngine/Public/FractureEngineConvex.h
--- a/Engine/Plugins/Experimental/Fracture/Source/FractureEngine/Public/FractureEngineConvex.h
+++ b/Engine/Plugins/Experimental/Fracture/Source/FractureEngine/Public/FractureEngineConvex.h
@@ -5,6 +5,7 @@
 #include "Containers/ArrayView.h"
 #include "CompGeom/ConvexDecomposition3.h"
 #include "DynamicMesh/DynamicMesh3.h"
+#include "UObject/ObjectMacros.h"
 
 namespace Chaos { class FConvex; }
 struct FManagedArrayCollection;
