exit the autocluster function early if there is no clustering to be done

fixes failed check when 0 clusters are requested

diff --git a/Engine/Plugins/Experimental/Fracture/Source/FractureEngine/Private/FractureEngineClustering.cpp b/Engine/Plugins/Experimental/Fracture/Source/FractureEngine/Private/FractureEngineClustering.cpp
--- a/Engine/Plugins/Experimental/Fracture/Source/FractureEngine/Private/FractureEngineClustering.cpp
+++ b/Engine/Plugins/Experimental/Fracture/Source/FractureEngine/Private/FractureEngineClustering.cpp
@@ -599,6 +599,12 @@ void FFractureEngineClustering::AutoCluster(FGeometryCollection& GeometryCollect
 		PartitionPositions = GenerateGridSites(GeometryCollection, ClusterIndex, InGridX, InGridY, InGridZ);
 	}
 
+	// Stop if we only want one cluster or there aren't enough children to do any clustering
+	if (DesiredSiteCountToUse <= 1 || NumChildren <= 1)
+	{
+		return;
+	}
+
 	int32 SiteCountToUse = DesiredSiteCountToUse;
 	int32 PreviousPartitionCount = TNumericLimits<int32>::Max();
 	bool bIterate = false;
