fix the ByGrid clustering not working in Fracture Mode AutoCluster tool

diff --git a/Engine/Plugins/Experimental/Fracture/Source/FractureEngine/Private/FractureEngineClustering.cpp b/Engine/Plugins/Experimental/Fracture/Source/FractureEngine/Private/FractureEngineClustering.cpp
--- a/Engine/Plugins/Experimental/Fracture/Source/FractureEngine/Private/FractureEngineClustering.cpp
+++ b/Engine/Plugins/Experimental/Fracture/Source/FractureEngine/Private/FractureEngineClustering.cpp
@@ -597,6 +597,7 @@ void FFractureEngineClustering::AutoCluster(FGeometryCollection& GeometryCollect
 	else if (ClusterSizeMethod == EFractureEngineClusterSizeMethod::ByGrid)
 	{
 		PartitionPositions = GenerateGridSites(GeometryCollection, ClusterIndex, InGridX, InGridY, InGridZ);
+		DesiredSiteCountToUse = PartitionPositions.Num();
 	}
 
 	// Stop if we only want one cluster or there aren't enough children to do any clustering
