Add the option to favor convexity in fracture mode autoclustering -- using the convex decomp algo's merge method to cluster, instead of k-means

diff --git a/Engine/Plugins/Experimental/ChaosEditor/Source/FractureEditor/Private/FractureToolAutoCluster.cpp b/Engine/Plugins/Experimental/ChaosEditor/Source/FractureEditor/Private/FractureToolAutoCluster.cpp
--- a/Engine/Plugins/Experimental/ChaosEditor/Source/FractureEditor/Private/FractureToolAutoCluster.cpp
+++ b/Engine/Plugins/Experimental/ChaosEditor/Source/FractureEditor/Private/FractureToolAutoCluster.cpp
@@ -206,6 +206,8 @@ void UFractureToolAutoCluster::Execute(TWeakPtr<FFractureEditorModeToolkit> InTo
 				for (const int32 ClusterIndex : Context.GetSelection())
 				{
 					int32 KMeansIterations = AutoClusterSettings->ClusterSizeMethod == EClusterSizeMethod::ByGrid ? AutoClusterSettings->DriftIterations : 500;
+					// note: bPreferConvexity is incompatible with grid clustering
+					bool bPreferConvexity = AutoClusterSettings->bPreferConvexity && AutoClusterSettings->ClusterSizeMethod != EClusterSizeMethod::ByGrid;
 					FFractureEngineClustering::AutoCluster(*GeometryCollection,
 						ClusterIndex,
 						(EFractureEngineClusterSizeMethod)AutoClusterSettings->ClusterSizeMethod,
@@ -219,7 +221,10 @@ void UFractureToolAutoCluster::Execute(TWeakPtr<FFractureEditorModeToolkit> InTo
 						AutoClusterSettings->ClusterGridDepth,
 						AutoClusterSettings->ClusterGridHeight,
 						AutoClusterSettings->MinimumSize,
-						KMeansIterations);
+						KMeansIterations, 
+						bPreferConvexity, 
+						AutoClusterSettings->ConcavityTolerance
+					);
 				}
 
 				Context.GenerateGuids(StartTransformCount);
diff --git a/Engine/Plugins/Experimental/ChaosEditor/Source/FractureEditor/Private/FractureToolAutoCluster.h b/Engine/Plugins/Experimental/ChaosEditor/Source/FractureEditor/Private/FractureToolAutoCluster.h
--- a/Engine/Plugins/Experimental/ChaosEditor/Source/FractureEditor/Private/FractureToolAutoCluster.h
+++ b/Engine/Plugins/Experimental/ChaosEditor/Source/FractureEditor/Private/FractureToolAutoCluster.h
@@ -90,12 +90,20 @@ public:
 	UPROPERTY(EditAnywhere, Category = ClusterSize, meta = (DisplayName = "Drift Iterations", ClampMin = "0", UIMax = "5", EditCondition = "ClusterSizeMethod == EClusterSizeMethod::ByGrid", EditConditionHides))
 	int DriftIterations = 0;
 
+	/** Whether to favor clusters that have a convex shape. (Note: Does not support ByGrid clustering.)  */
+	UPROPERTY(EditAnywhere, Category = AutoCluster, meta = (EditCondition = "ClusterSizeMethod != EClusterSizeMethod::ByGrid"))
+	bool bPreferConvexity = false;
+
+	/** If > 0, cube root of maximum concave volume to add per cluster (ignoring concavity of individual parts) */
+	UPROPERTY(EditAnywhere, Category = AutoCluster, meta = (EditCondition = "bPreferConvexity && ClusterSizeMethod != EClusterSizeMethod::ByGrid"))
+	double ConcavityTolerance = 0;
+
 	/** If true, bones will only be added to the same cluster if they are physically connected (either directly, or via other bones in the same cluster) */
 	UPROPERTY(EditAnywhere, Category = AutoCluster, meta = (DisplayName = "Enforce Cluster Connectivity"))
 	bool bEnforceConnectivity=true;
 
 	/** If true, make sure the site parameters are matched as close as possible ( bEnforceConnectivity can make the number of site larger than the requested input may produce without it ) */
-	UPROPERTY(EditAnywhere, Category = AutoCluster, meta = (EditCondition = "bEnforceConnectivity == true && ClusterSizeMethod != EClusterSizeMethod::ByGrid"))
+	UPROPERTY(EditAnywhere, Category = AutoCluster, meta = (EditCondition = "bEnforceConnectivity == true && ClusterSizeMethod != EClusterSizeMethod::ByGrid && !bPreferConvexity"))
 	bool bEnforceSiteParameters = true;
 
 	/** If true, prevent the creation of clusters with only a single child, skipping creation of a new cluster in such cases. */
diff --git a/Engine/Plugins/Experimental/Fracture/Source/FractureEngine/Private/FractureEngineClustering.cpp b/Engine/Plugins/Experimental/Fracture/Source/FractureEngine/Private/FractureEngineClustering.cpp
--- a/Engine/Plugins/Experimental/Fracture/Source/FractureEngine/Private/FractureEngineClustering.cpp
+++ b/Engine/Plugins/Experimental/Fracture/Source/FractureEngine/Private/FractureEngineClustering.cpp
@@ -9,6 +9,121 @@
 #include "GeometryCollection/GeometryCollectionClusteringUtility.h"
 #include "GeometryCollection/GeometryCollectionAlgo.h"
 #include "Math/BoxSphereBounds.h"
+#include "VertexConnectedComponents.h"
+#include "CompGeom/ConvexDecomposition3.h"
+#include "CompGeom/ConvexHull3.h"
 
+namespace UE::Private::ClusterMagnet
+{
+	struct FClusterMagnet
+	{
+		TSet<int32> ClusteredNodes;
+		TSet<int32> Connections;
+	};
+
+	TMap<int32, TSet<int32>> InitializeConnectivity(const TSet<int32>& TopNodes, FGeometryCollection* GeometryCollection, int32 OperatingLevel);
+	void CollectTopNodeConnections(FGeometryCollection* GeometryCollection, int32 Index, int32 OperatingLevel, TSet<int32>& OutConnections);
+	void SeparateClusterMagnets(const TSet<int32>& TopNodes, const TArray<int32>& Selection, const TMap<int32, TSet<int32>>& TopNodeConnectivity, TArray<FClusterMagnet>& OutClusterMagnets, TSet<int32>& OutRemainingPool);
+	bool AbsorbClusterNeighbors(const TMap<int32, TSet<int32>> TopNodeConnectivity, FClusterMagnet& OutClusterMagnets, TSet<int32>& OutRemainingPool);
+
+
+	TMap<int32, TSet<int32>> InitializeConnectivity(const TSet<int32>& TopNodes, FGeometryCollection* GeometryCollection, int32 OperatingLevel)
+	{
+		FGeometryCollectionProximityUtility ProximityUtility(GeometryCollection);
+		ProximityUtility.RequireProximity();
+
+		TMap<int32, TSet<int32>> ConnectivityMap;
+		for (int32 Index : TopNodes)
+		{
+			// Collect the proximity indices of all the leaf nodes under this top node,
+			// traced back up to its parent top node, so that all connectivity describes
+			// relationships only between top nodes.
+			TSet<int32> Connections;
+			CollectTopNodeConnections(GeometryCollection, Index, OperatingLevel, Connections);
+			Connections.Remove(Index);
+
+			// Remove any connections outside the current operating branch.
+			ConnectivityMap.Add(Index, Connections.Intersect(TopNodes));
+		}
+
+		return ConnectivityMap;
+	}
+
+	void CollectTopNodeConnections(FGeometryCollection* GeometryCollection, int32 Index, int32 OperatingLevel, TSet<int32>& OutConnections)
+	{
+		const TManagedArray<int32>& TransformToGeometryIndex = GeometryCollection->TransformToGeometryIndex;
+		if (GeometryCollection->SimulationType[Index] == FGeometryCollection::ESimulationTypes::FST_Rigid
+			&& TransformToGeometryIndex[Index] != INDEX_NONE) // rigid node with geometry, leaf of the simulated part
+		{
+			const TManagedArray<TSet<int32>>& Proximity = GeometryCollection->GetAttribute<TSet<int32>>("Proximity", FGeometryCollection::GeometryGroup);
+			const TManagedArray<int32>& GeometryToTransformIndex = GeometryCollection->TransformIndex;
+
+
+			for (int32 Neighbor : Proximity[TransformToGeometryIndex[Index]])
+			{
+				int32 NeighborTransformIndex = GeometryToTransformIndex[Neighbor];
+				OutConnections.Add(FGeometryCollectionClusteringUtility::GetParentOfBoneAtSpecifiedLevel(GeometryCollection, NeighborTransformIndex, OperatingLevel));
+			}
+		}
+		else
+		{
+			const TManagedArray<TSet<int32>>& Children = GeometryCollection->Children;
+			for (int32 ChildIndex : Children[Index])
+			{
+				CollectTopNodeConnections(GeometryCollection, ChildIndex, OperatingLevel, OutConnections);
+			}
+		}
+	}
+
+	void SeparateClusterMagnets(
+		const TSet<int32>& TopNodes,
+		const TArray<int32>& Selection,
+		const TMap<int32, TSet<int32>>& TopNodeConnectivity,
+		TArray<FClusterMagnet>& OutClusterMagnets,
+		TSet<int32>& OutRemainingPool)
+	{
+		OutClusterMagnets.Reserve(TopNodes.Num());
+		OutRemainingPool.Reserve(TopNodes.Num());
+
+		for (int32 Index : TopNodes)
+		{
+			if (Selection.Contains(Index))
+			{
+				OutClusterMagnets.AddDefaulted();
+				FClusterMagnet& NewMagnet = OutClusterMagnets.Last();
+				NewMagnet.ClusteredNodes.Add(Index);
+				NewMagnet.Connections = TopNodeConnectivity[Index];
+			}
+			else
+			{
+				OutRemainingPool.Add(Index);
+			}
+		}
+	}
+
+	bool AbsorbClusterNeighbors(const TMap<int32, TSet<int32>> TopNodeConnectivity, FClusterMagnet& OutClusterMagnet, TSet<int32>& OutRemainingPool)
+	{
+		// Return true if neighbors were absorbed.
+		bool bNeighborsAbsorbed = false;
+
+		TSet<int32> NewConnections;
+		for (int32 NeighborIndex : OutClusterMagnet.Connections)
+		{
+			// If the neighbor is still in the pool, absorb it and its connections.
+			if (OutRemainingPool.Contains(NeighborIndex))
+			{
+				OutClusterMagnet.ClusteredNodes.Add(NeighborIndex);
+				NewConnections.Append(TopNodeConnectivity[NeighborIndex]);
+				OutRemainingPool.Remove(NeighborIndex);
+				bNeighborsAbsorbed = true;
+			}
+		}
+		OutClusterMagnet.Connections.Append(NewConnections);
+
+		return bNeighborsAbsorbed;
+	}
+}
+
+
 FVoronoiPartitioner::FVoronoiPartitioner(const FGeometryCollection* GeometryCollection, int32 ClusterIndex)
 {
@@ -544,7 +659,9 @@
 	const int32 GridY,
 	const int32 GridZ,
 	const float MinimumClusterSize,
-	const int32 KMeansIterations)
+	const int32 KMeansIterations,
+	const bool bPreferConvexity,
+	const float ConcavityTolerance)
 {
 	FFractureEngineBoneSelection Selection(GeometryCollection, BoneIndices);
 
@@ -553,10 +670,202 @@
 	for (const int32 ClusterIndex : Selection.GetSelectedBones())
 	{
 		AutoCluster(GeometryCollection, ClusterIndex, ClusterSizeMethod, SiteCount, SiteCountFraction, SiteSize, 
-			bEnforceConnectivity, bAvoidIsolated, bEnforceSiteParameters, GridX, GridY, GridZ, MinimumClusterSize, KMeansIterations);
+			bEnforceConnectivity, bAvoidIsolated, bEnforceSiteParameters, GridX, GridY, GridZ, MinimumClusterSize, 
+			KMeansIterations, bPreferConvexity, ConcavityTolerance);
 	}
 }
 
+void FFractureEngineClustering::ConvexityBasedCluster(FGeometryCollection& GeometryCollection,
+	int32 ClusterIndex,
+	uint32 SiteCount,
+	bool bEnforceConnectivity,
+	bool bAvoidIsolated,
+	float ConcavityTolerance
+)
+{
+	const bool bUseVolumesOfConvexHulls = true;
+	if (bUseVolumesOfConvexHulls)
+	{
+		FGeometryCollectionConvexUtility::SetVolumeAttributes(&GeometryCollection);
+	}
+	FGeometryCollectionProximityUtility ProximityUtility(&GeometryCollection);
+	ProximityUtility.UpdateProximity();
+	const TManagedArray<float>& GeometryVolumes = GeometryCollection.GetAttribute<float>("Volume", FTransformCollection::TransformGroup);
+
+
+	UE::Geometry::FConvexDecomposition3 ConvexDecomp;
+	const TSet<int32>& ChildrenSet = GeometryCollection.Children[ClusterIndex];
+	const TArray<int32> Children = ChildrenSet.Array();
+	UE::Geometry::FSizedDisjointSet Components(Children.Num());
+
+	Chaos::Facades::FCollectionHierarchyFacade HierarchyFacade(GeometryCollection);
+	HierarchyFacade.GenerateLevelAttribute();
+	const TManagedArray<int32>& Levels = GeometryCollection.GetAttribute<int32>("Level", FGeometryCollection::TransformGroup);
+	int32 Level = Levels[ClusterIndex];
+
+	// use source geometry as input convex vertices
+	TArray<int32> TransformVertexCounts;
+	TArray<int32> TransformVertexStarts;
+	TransformVertexCounts.SetNumUninitialized(Children.Num());
+	TransformVertexStarts.SetNumUninitialized(Children.Num());
+	TArray<double> Volumes;
+	Volumes.SetNumZeroed(Children.Num());
+	TArray<FVector3d> GlobalVertices;
+	TArray<int32> ToProcess;
+	TArray<FTransform> GlobalTransforms;
+
+	if (!bUseVolumesOfConvexHulls)
+	{
+		for (int32 ChildIdx = 0; ChildIdx < Children.Num(); ++ChildIdx)
+		{
+			Volumes[ChildIdx] = GeometryVolumes[Children[ChildIdx]];
+		}
+	}
+
+	GeometryCollectionAlgo::GlobalMatrices(GeometryCollection.Transform, GeometryCollection.Parent, GlobalTransforms);
+	for (int32 ChildIdx = 0; ChildIdx < Children.Num(); ++ChildIdx)
+	{
+		TransformVertexStarts[ChildIdx] = GlobalVertices.Num();
+		int32 TransformIdx = Children[ChildIdx];
+		ToProcess.Reset();
+		ToProcess.Add(TransformIdx);
+		double Volume = 0;
+		while (!ToProcess.IsEmpty())
+		{
+			int32 ProcessTransformIdx = ToProcess.Pop(false);
+			FGeometryCollection::ESimulationTypes SimType = (FGeometryCollection::ESimulationTypes)GeometryCollection.SimulationType[ProcessTransformIdx];
+			if (SimType == FGeometryCollection::ESimulationTypes::FST_Clustered)
+			{
+				for (int32 ProcChild : GeometryCollection.Children[ProcessTransformIdx])
+				{
+					ToProcess.Add(ProcChild);
+				}
+			}
+			else if (SimType == FGeometryCollection::ESimulationTypes::FST_Rigid)
+			{
+				int32 ProcessGeometryIdx = GeometryCollection.TransformToGeometryIndex[ProcessTransformIdx];
+				if (ProcessGeometryIdx == INDEX_NONE)
+				{
+					continue;
+				}
+				int32 GCVStart = GeometryCollection.VertexStart[ProcessGeometryIdx];
+				int32 GCVEnd = GCVStart + GeometryCollection.VertexCount[ProcessGeometryIdx];
+				int32 GlobalStart = GlobalVertices.Num();
+				for (int32 VIdx = GCVStart; VIdx < GCVEnd; ++VIdx)
+				{
+					GlobalVertices.Add(GlobalTransforms[ProcessTransformIdx].TransformPosition((FVector3d)GeometryCollection.Vertex[VIdx]));
+				}
+				if (bUseVolumesOfConvexHulls)
+				{
+					double HullVolume = UE::Geometry::FConvexHull3d::ComputeVolume(TArrayView<const FVector3d>(&GlobalVertices[GlobalStart], GlobalVertices.Num() - GlobalStart));
+					Volumes[ChildIdx] += HullVolume;
+				}
+			}
+		}
+		TransformVertexCounts[ChildIdx] = GlobalVertices.Num() - TransformVertexStarts[ChildIdx];
+	}
+
+	TArray<TPair<int32, int32>> ChildrenProximity;
+	if (bEnforceConnectivity)
+	{
+		TMap<int32, TSet<int32>> Connectivity = UE::Private::ClusterMagnet::InitializeConnectivity(ChildrenSet, &GeometryCollection, Level + 1);
+		TMap<int32, int32> TransformIdxToChildrenIdx;
+		for (int32 LocalIdx = 0; LocalIdx < Children.Num(); ++LocalIdx)
+		{
+			TransformIdxToChildrenIdx.Add(Children[LocalIdx], LocalIdx);
+		}
+		for (const TPair<int32, TSet<int32>>& ConnectionsA : Connectivity)
+		{
+			int32 ChildIdxA = TransformIdxToChildrenIdx[ConnectionsA.Key];
+			const TSet<int32>& ConnectedTo = ConnectionsA.Value;
+			for (int32 TransformB : ConnectedTo)
+			{
+				int32 ChildIdxB = TransformIdxToChildrenIdx[TransformB];
+				ChildrenProximity.Emplace(ChildIdxA, ChildIdxB);
+			}
+		}
+	}
+
+	ConvexDecomp.InitializeFromHulls(Children.Num(),
+		[&](int32 Idx)->double { return Volumes[Idx]; },
+		[&](int32 Idx)->int32 { return TransformVertexCounts[Idx]; },
+		[&](int32 Idx, int32 VertIdx) -> FVector3d { return GlobalVertices[TransformVertexStarts[Idx] + VertIdx]; },
+		ChildrenProximity);
+
+	if (!bEnforceConnectivity)
+	{
+		// If we're not getting connectivity from the geometry collection, add it based on bounding box overlaps
+		ConvexDecomp.InitializeProximityFromDecompositionBoundingBoxOverlaps(1, .5, 1);
+	}
+
+	UE::Geometry::FConvexDecomposition3::FMergeSettings MergeSettings;
+	MergeSettings.TargetNumParts = SiteCount;
+
+	double TotalVolume = 0;
+	for (double Volume : Volumes)
+	{
+		TotalVolume += Volume;
+	}
+	MergeSettings.ErrorTolerance = ConcavityTolerance;
+	MergeSettings.bErrorToleranceOverridesNumParts = false;
+	double MaxVolumeToMerge = TotalVolume / SiteCount; // stop merging parts that are larger than their share of the volume ...
+	double MinVolumeToSkip = .2 * TotalVolume / SiteCount; // except still allow merging very small parts into large parts
+	MergeSettings.CustomAllowMergeParts =
+		[MaxVolumeToMerge, MinVolumeToSkip](const UE::Geometry::FConvexDecomposition3::FConvexPart& PartA, const UE::Geometry::FConvexDecomposition3::FConvexPart& PartB)
+		{
+			if (PartA.bMustMerge || PartB.bMustMerge)
+			{
+				return true;
+			}
+			return ((PartA.SumHullsVolume < MaxVolumeToMerge && PartB.SumHullsVolume < MaxVolumeToMerge) ||
+					(PartA.SumHullsVolume < MinVolumeToSkip || PartB.SumHullsVolume < MinVolumeToSkip));
+		};
+	MergeSettings.MergeCallback = [&Components](int32 ChildA, int32 ChildB)
+	{
+		Components.Union(ChildA, ChildB);
+	};
+	
+	ConvexDecomp.MergeBest(MergeSettings);
+	TArray<int32> CompactIdxToGroupID, GroupIDToCompactIdx;
+	int32 PartitionCount = Components.CompactedGroupIndexToGroupID(&CompactIdxToGroupID, &GroupIDToCompactIdx);
+	TArray<int32> NewCluster;
+	int32 NewClusterIndexStart = GeometryCollection.AddElements(PartitionCount, FGeometryCollection::TransformGroup);
+	bool bHasEmptyClusters = false;
+	for (int32 Index = 0; Index < PartitionCount; ++Index)
+	{
+		NewCluster.Reset();
+		int32 GroupID = CompactIdxToGroupID[Index];
+		for (int32 ChildIdx = 0; ChildIdx < Children.Num(); ++ChildIdx)
+		{
+			if (Components.Find(ChildIdx) == GroupID)
+			{
+				NewCluster.Add(Children[ChildIdx]);
+			}
+		}
+		if (bAvoidIsolated && NewCluster.Num() == 1)
+		{
+			bHasEmptyClusters = true;
+			NewCluster.Reset();
+		}
+
+		int32 NewClusterIndex = NewClusterIndexStart + Index;
+		GeometryCollection.Parent[NewClusterIndex] = ClusterIndex;
+		GeometryCollection.Children[ClusterIndex].Add(NewClusterIndex);
+		GeometryCollection.BoneName[NewClusterIndex] = "ClusterBone";
+		GeometryCollection.Children[NewClusterIndex] = TSet<int32>(NewCluster);
+		GeometryCollection.SimulationType[NewClusterIndex] = FGeometryCollection::ESimulationTypes::FST_Clustered;
+		GeometryCollection.Transform[NewClusterIndex] = FTransform3f::Identity;
+		GeometryCollectionAlgo::ParentTransforms(&GeometryCollection, NewClusterIndex, NewCluster);
+	}
+	if (bHasEmptyClusters)
+	{
+		FGeometryCollectionClusteringUtility::RemoveDanglingClusters(&GeometryCollection);
+	}
+	FGeometryCollectionClusteringUtility::UpdateHierarchyLevelOfChildren(&GeometryCollection, ClusterIndex);
+	FGeometryCollectionClusteringUtility::RecursivelyUpdateChildBoneNames(ClusterIndex, GeometryCollection.Children, GeometryCollection.BoneName);
+	FGeometryCollectionClusteringUtility::ValidateResults(&GeometryCollection);
+}
+
 void FFractureEngineClustering::AutoCluster(FGeometryCollection& GeometryCollection,
 	const int32 ClusterIndex,
 	const EFractureEngineClusterSizeMethod ClusterSizeMethod,
@@ -570,16 +879,16 @@
 	const int32 InGridY,
 	const int32 InGridZ,
 	const float MinimumClusterSize,
-	const int32 KMeansIterations)
+	const int32 KMeansIterations,
+	const bool bPreferConvexity,
+	const float ConcavityTolerance)
 {
-	FVoronoiPartitioner VoronoiPartition(&GeometryCollection, ClusterIndex);
-	int32 NumChildren = GeometryCollection.Children[ClusterIndex].Num();
-
 	// Used by the ByGrid method, which manually distributes initial positions
 	TArray<FVector> PartitionPositions;
 
 	int32 DesiredSiteCountToUse = 1;
 	int32 IterationsToUse = KMeansIterations;
+	int32 NumChildren = GeometryCollection.Children[ClusterIndex].Num();
 	if (ClusterSizeMethod == EFractureEngineClusterSizeMethod::ByNumber)
 	{
 		DesiredSiteCountToUse = SiteCount;
@@ -600,12 +909,26 @@
 		DesiredSiteCountToUse = PartitionPositions.Num();
 	}
 
+	if (bPreferConvexity)
+	{
+		ConvexityBasedCluster(GeometryCollection, ClusterIndex, SiteCount, bEnforceConnectivity, bAvoidIsolated, ConcavityTolerance);
+		return;
+	}
+
+	FVoronoiPartitioner VoronoiPartition(&GeometryCollection, ClusterIndex);
+
 	// Stop if we only want one cluster or there aren't enough children to do any clustering
 	if (DesiredSiteCountToUse <= 1 || NumChildren <= 1)
 	{
 		return;
 	}
 
+	bool bNeedsVolume = MinimumClusterSize > 0;
+	if (bNeedsVolume)
+	{
+		FGeometryCollectionConvexUtility::SetVolumeAttributes(&GeometryCollection);
+	}
+
 	int32 SiteCountToUse = DesiredSiteCountToUse;
 	int32 PreviousPartitionCount = TNumericLimits<int32>::Max();
 	bool bIterate = false;
@@ -852,116 +1175,6 @@
 	return false;
 }
 
-namespace UE::Private::ClusterMagnet
-{
-	struct FClusterMagnet
-	{
-		TSet<int32> ClusteredNodes;
-		TSet<int32> Connections;
-	};
 
-	TMap<int32, TSet<int32>> InitializeConnectivity(const TSet<int32>& TopNodes, FGeometryCollection* GeometryCollection, int32 OperatingLevel);
-	void CollectTopNodeConnections(FGeometryCollection* GeometryCollection, int32 Index, int32 OperatingLevel, TSet<int32>& OutConnections);
-	void SeparateClusterMagnets(const TSet<int32>& TopNodes, const TArray<int32>& Selection, const TMap<int32, TSet<int32>>& TopNodeConnectivity, TArray<FClusterMagnet>& OutClusterMagnets, TSet<int32>& OutRemainingPool);
-	bool AbsorbClusterNeighbors(const TMap<int32, TSet<int32>> TopNodeConnectivity, FClusterMagnet& OutClusterMagnets, TSet<int32>& OutRemainingPool);
-
-
-	TMap<int32, TSet<int32>> InitializeConnectivity(const TSet<int32>& TopNodes, FGeometryCollection* GeometryCollection, int32 OperatingLevel)
-	{
-		FGeometryCollectionProximityUtility ProximityUtility(GeometryCollection);
-		ProximityUtility.RequireProximity();
-
-		TMap<int32, TSet<int32>> ConnectivityMap;
-		for (int32 Index : TopNodes)
-		{
-			// Collect the proximity indices of all the leaf nodes under this top node,
-			// traced back up to its parent top node, so that all connectivity describes
-			// relationships only between top nodes.
-			TSet<int32> Connections;
-			CollectTopNodeConnections(GeometryCollection, Index, OperatingLevel, Connections);
-			Connections.Remove(Index);
-
-			// Remove any connections outside the current operating branch.
-			ConnectivityMap.Add(Index, Connections.Intersect(TopNodes));
-		}
-
-		return ConnectivityMap;
-	}
-
-	void CollectTopNodeConnections(FGeometryCollection* GeometryCollection, int32 Index, int32 OperatingLevel, TSet<int32>& OutConnections)
-	{
-		const TManagedArray<int32>& TransformToGeometryIndex = GeometryCollection->TransformToGeometryIndex;
-		if (GeometryCollection->SimulationType[Index] == FGeometryCollection::ESimulationTypes::FST_Rigid
-			&& TransformToGeometryIndex[Index] != INDEX_NONE) // rigid node with geometry, leaf of the simulated part
-		{
-			const TManagedArray<TSet<int32>>& Proximity = GeometryCollection->GetAttribute<TSet<int32>>("Proximity", FGeometryCollection::GeometryGroup);
-			const TManagedArray<int32>& GeometryToTransformIndex = GeometryCollection->TransformIndex;
-
-
-			for (int32 Neighbor : Proximity[TransformToGeometryIndex[Index]])
-			{
-				int32 NeighborTransformIndex = GeometryToTransformIndex[Neighbor];
-				OutConnections.Add(FGeometryCollectionClusteringUtility::GetParentOfBoneAtSpecifiedLevel(GeometryCollection, NeighborTransformIndex, OperatingLevel));
-			}
-		}
-		else
-		{
-			const TManagedArray<TSet<int32>>& Children = GeometryCollection->Children;
-			for (int32 ChildIndex : Children[Index])
-			{
-				CollectTopNodeConnections(GeometryCollection, ChildIndex, OperatingLevel, OutConnections);
-			}
-		}
-	}
-
-	void SeparateClusterMagnets(
-		const TSet<int32>& TopNodes,
-		const TArray<int32>& Selection,
-		const TMap<int32, TSet<int32>>& TopNodeConnectivity,
-		TArray<FClusterMagnet>& OutClusterMagnets,
-		TSet<int32>& OutRemainingPool)
-	{
-		OutClusterMagnets.Reserve(TopNodes.Num());
-		OutRemainingPool.Reserve(TopNodes.Num());
-
-		for (int32 Index : TopNodes)
-		{
-			if (Selection.Contains(Index))
-			{
-				OutClusterMagnets.AddDefaulted();
-				FClusterMagnet& NewMagnet = OutClusterMagnets.Last();
-				NewMagnet.ClusteredNodes.Add(Index);
-				NewMagnet.Connections = TopNodeConnectivity[Index];
-			}
-			else
-			{
-				OutRemainingPool.Add(Index);
-			}
-		}
-	}
-
-	bool AbsorbClusterNeighbors(const TMap<int32, TSet<int32>> TopNodeConnectivity, FClusterMagnet& OutClusterMagnet, TSet<int32>& OutRemainingPool)
-	{
-		// Return true if neighbors were absorbed.
-		bool bNeighborsAbsorbed = false;
-
-		TSet<int32> NewConnections;
-		for (int32 NeighborIndex : OutClusterMagnet.Connections)
-		{
-			// If the neighbor is still in the pool, absorb it and its connections.
-			if (OutRemainingPool.Contains(NeighborIndex))
-			{
-				OutClusterMagnet.ClusteredNodes.Add(NeighborIndex);
-				NewConnections.Append(TopNodeConnectivity[NeighborIndex]);
-				OutRemainingPool.Remove(NeighborIndex);
-				bNeighborsAbsorbed = true;
-			}
-		}
-		OutClusterMagnet.Connections.Append(NewConnections);
-
-		return bNeighborsAbsorbed;
-	}
-}
-
 bool FFractureEngineClustering::ClusterMagnet(
 	FGeometryCollection& GeometryCollection,
diff --git a/Engine/Plugins/Experimental/Fracture/Source/FractureEngine/Public/FractureEngineClustering.h b/Engine/Plugins/Experimental/Fracture/Source/FractureEngine/Public/FractureEngineClustering.h
--- a/Engine/Plugins/Experimental/Fracture/Source/FractureEngine/Public/FractureEngineClustering.h
+++ b/Engine/Plugins/Experimental/Fracture/Source/FractureEngine/Public/FractureEngineClustering.h
@@ -99,7 +99,9 @@ public:
 		const int32 GridY = 2,
 		const int32 GridZ = 2,
 		const float MinimumClusterSize = 0,
-		const int32 KMeansIterations = 500);
+		const int32 KMeansIterations = 500,
+		const bool bPreferConvexity = false,
+		const float ConcavityErrorTolerance = 0);
 
 	static void AutoCluster(FGeometryCollection& GeometryCollection,
 		const int32 ClusterIndex,
@@ -114,7 +116,17 @@ public:
 		const int32 GridY = 2,
 		const int32 GridZ = 2,
 		const float MinimumClusterSize = 0,
-		const int32 KMeansIterations = 500);
+		const int32 KMeansIterations = 500,
+		const bool bPreferConvexity = false,
+		const float ConcavityErrorTolerance = 0);
+
+	// Autoclustering that favors convex-shaped clusters
+	static void ConvexityBasedCluster(FGeometryCollection& GeometryCollection,
+		int32 ClusterIndex,
+		uint32 SiteCount,
+		bool bEnforceConnectivity,
+		bool bAvoidIsolated,
+		float ConcavityErrorTolerance);
 
 	static TArray<FVector> GenerateGridSites(
 		const FGeometryCollection& GeometryCollection,
