fix deprecated allow-shrinking bool flag in fracture engine clustering

diff --git a/Engine/Plugins/Experimental/Fracture/Source/FractureEngine/Private/FractureEngineClustering.cpp b/Engine/Plugins/Experimental/Fracture/Source/FractureEngine/Private/FractureEngineClustering.cpp
--- a/Engine/Plugins/Experimental/Fracture/Source/FractureEngine/Private/FractureEngineClustering.cpp
+++ b/Engine/Plugins/Experimental/Fracture/Source/FractureEngine/Private/FractureEngineClustering.cpp
@@ -732,7 +732,7 @@ void FFractureEngineClustering::ConvexityBasedCluster(FGeometryCollection& Geome
 		double Volume = 0;
 		while (!ToProcess.IsEmpty())
 		{
-			int32 ProcessTransformIdx = ToProcess.Pop(false);
+			int32 ProcessTransformIdx = ToProcess.Pop(EAllowShrinking::No);
 			FGeometryCollection::ESimulationTypes SimType = (FGeometryCollection::ESimulationTypes)GeometryCollection.SimulationType[ProcessTransformIdx];
 			if (SimType == FGeometryCollection::ESimulationTypes::FST_Clustered)
 			{
