Fix wrong Log category being declare twice

diff --git a/Engine/Plugins/Experimental/Fracture/Source/FractureEngine/Private/FractureEngineUtility.cpp b/Engine/Plugins/Experimental/Fracture/Source/FractureEngine/Private/FractureEngineUtility.cpp
--- a/Engine/Plugins/Experimental/Fracture/Source/FractureEngine/Private/FractureEngineUtility.cpp
+++ b/Engine/Plugins/Experimental/Fracture/Source/FractureEngine/Private/FractureEngineUtility.cpp
@@ -332,8 +332,8 @@ static void AddSingleRootNodeIfRequired(TUniquePtr<FGeometryCollection>& InOutCo
 	}
 }
 
-DECLARE_LOG_CATEGORY_EXTERN(LogFractureTool, Log, All);
-DEFINE_LOG_CATEGORY(LogFractureTool);
+DECLARE_LOG_CATEGORY_EXTERN(LogFractureEngineUtility, Log, All);
+DEFINE_LOG_CATEGORY(LogFractureEngineUtility);
 
 void FFractureEngineUtility::ValidateGeometryCollection(FManagedArrayCollection& InOutCollection,
 	const bool InRemoveUnreferencedGeometry,
@@ -355,7 +355,7 @@ void FFractureEngineUtility::ValidateGeometryCollection(FManagedArrayCollection&
 				if (GeometryCollection->IsClustered(Idx) && TransformToGeometry[Idx] != INDEX_NONE)
 				{
 					TransformToGeometry[Idx] = INDEX_NONE;
-					UE_LOG(LogFractureTool, Verbose, TEXT("Removed geometry index from cluster %d."), Idx);
+					UE_LOG(LogFractureEngineUtility, Verbose, TEXT("Removed geometry index from cluster %d."), Idx);
 					bDirty = true;
 				}
 			}
@@ -375,7 +375,7 @@ void FFractureEngineUtility::ValidateGeometryCollection(FManagedArrayCollection&
 				if ((TransformIndex[Idx] == INDEX_NONE) || (TransformToGeometry[TransformIndex[Idx]] != Idx))
 				{
 					RemoveGeometry.Add(Idx);
-					UE_LOG(LogFractureTool, Verbose, TEXT("Removed dangling geometry at index %d."), Idx);
+					UE_LOG(LogFractureEngineUtility, Verbose, TEXT("Removed dangling geometry at index %d."), Idx);
 					bDirty = true;
 				}
 			}
@@ -392,7 +392,7 @@ void FFractureEngineUtility::ValidateGeometryCollection(FManagedArrayCollection&
 		{
 			if (FGeometryCollectionClusteringUtility::RemoveClustersOfOnlyOneChild(GeometryCollection.Get()))
 			{
-				UE_LOG(LogFractureTool, Verbose, TEXT("Removed one or more clusters of only one child."));
+				UE_LOG(LogFractureEngineUtility, Verbose, TEXT("Removed one or more clusters of only one child."));
 				bDirty = true;
 			}
 		}
@@ -401,7 +401,7 @@ void FFractureEngineUtility::ValidateGeometryCollection(FManagedArrayCollection&
 		{
 			if (FGeometryCollectionClusteringUtility::RemoveDanglingClusters(GeometryCollection.Get()))
 			{
-				UE_LOG(LogFractureTool, Verbose, TEXT("Removed one or more dangling clusters."));
+				UE_LOG(LogFractureEngineUtility, Verbose, TEXT("Removed one or more dangling clusters."));
 				bDirty = true;
 			}
 		}
