Fix various crashes happening when evaluating cutting and merging dataflow node with empty collections

diff --git a/Engine/Plugins/Experimental/Fracture/Source/FractureEngine/Private/FractureEngineClustering.cpp b/Engine/Plugins/Experimental/Fracture/Source/FractureEngine/Private/FractureEngineClustering.cpp
--- a/Engine/Plugins/Experimental/Fracture/Source/FractureEngine/Private/FractureEngineClustering.cpp
+++ b/Engine/Plugins/Experimental/Fracture/Source/FractureEngine/Private/FractureEngineClustering.cpp
@@ -1113,6 +1113,12 @@ bool FFractureEngineClustering::MergeSelectedClusters(FGeometryCollection& Geome
 {
 	Chaos::Facades::FCollectionHierarchyFacade HierarchyFacade(GeometryCollection);
 
+	const int32 NumTransforms = GeometryCollection.NumElements(FGeometryCollection::TransformGroup);
+	if (Selection.Num() > NumTransforms)
+	{
+		return false;
+	}
+
 	GeometryCollection::Facades::FCollectionTransformSelectionFacade SelectionFacade(GeometryCollection);
 	SelectionFacade.ConvertEmbeddedSelectionToParents(Selection); // embedded geo must stay attached to parent
 
diff --git a/Engine/Plugins/Experimental/Fracture/Source/FractureEngine/Private/FractureEngineFracturing.cpp b/Engine/Plugins/Experimental/Fracture/Source/FractureEngine/Private/FractureEngineFracturing.cpp
--- a/Engine/Plugins/Experimental/Fracture/Source/FractureEngine/Private/FractureEngineFracturing.cpp
+++ b/Engine/Plugins/Experimental/Fracture/Source/FractureEngine/Private/FractureEngineFracturing.cpp
@@ -983,6 +983,13 @@ int32 FFractureEngineFracturing::MeshCutter(TArray<FTransform>& MeshTransforms,
 	const bool InSplitIslands,
 	const float InCollisionSampleSpacing)
 {
+
+	if (InOutCollection.NumElements(FTransformCollection::TransformGroup) == 0)
+	{
+		// empty collection, early exit
+		return -1;
+	}
+
 	if (TUniquePtr<FGeometryCollection> GeomCollection = TUniquePtr<FGeometryCollection>(InOutCollection.NewCopy<FGeometryCollection>()))
 	{
 		// Note: Noise not currently supported
diff --git a/Engine/Plugins/Experimental/PlanarCutPlugin/Source/PlanarCut/Private/PlanarCut.cpp b/Engine/Plugins/Experimental/PlanarCutPlugin/Source/PlanarCut/Private/PlanarCut.cpp
--- a/Engine/Plugins/Experimental/PlanarCutPlugin/Source/PlanarCut/Private/PlanarCut.cpp
+++ b/Engine/Plugins/Experimental/PlanarCutPlugin/Source/PlanarCut/Private/PlanarCut.cpp
@@ -1700,6 +1700,12 @@ void MergeClusters(
 	bool bOnlySameParent
 )
 {
+	if (Collection.NumElements(FGeometryCollection::TransformGroup) == 0)
+	{
+		// nothing to do, early exit
+		return;
+	}
+
 	FGeometryCollectionProximityUtility ProximityUtility(&Collection);
 	ProximityUtility.RequireProximity();
 	const TManagedArray<TSet<int32>>& Proximity = Collection.GetAttribute<TSet<int32>>("Proximity", FGeometryCollection::GeometryGroup);
