Dataflow:
- Added FUniformFractureDataflowNode, it matches the Uniform fracture tool 100%
- Fixed GUID display in CollectionSpreadsheet
- Bug fixes

diff --git a/Engine/Plugins/Experimental/Dataflow/Source/DataflowEditor/Private/Dataflow/CollectionSpreadSheetWidget.cpp b/Engine/Plugins/Experimental/Dataflow/Source/DataflowEditor/Private/Dataflow/CollectionSpreadSheetWidget.cpp
--- a/Engine/Plugins/Experimental/Dataflow/Source/DataflowEditor/Private/Dataflow/CollectionSpreadSheetWidget.cpp
+++ b/Engine/Plugins/Experimental/Dataflow/Source/DataflowEditor/Private/Dataflow/CollectionSpreadSheetWidget.cpp
@@ -120,6 +120,11 @@ namespace CollectionSpreadSheetHelpers
 		return FString::Printf(TEXT("%d %d %d %d"), Value[0], Value[1], Value[2], Value[3]);
 	}
 
+	FString AttributeValueToString(FGuid Value)
+	{
+		return Value.ToString();
+	}
+
 	template<typename T>
 	FString AttributeValueToString(const FManagedArrayCollection& InCollection, const FName& InAttributeName, const FName& InGroupName, int32 InIdxColumn)
 	{
@@ -213,6 +218,10 @@ namespace CollectionSpreadSheetHelpers
 			ValueAsString = AttributeValueToString<FIntVector4>(InCollection, InAttributeName, InGroupName, InIdxColumn);
 			break;
 
+		case FManagedArrayCollection::EArrayType::FGuidType:
+			ValueAsString = AttributeValueToString<FGuid>(InCollection, InAttributeName, InGroupName, InIdxColumn);
+			break;
+
 		default:
 			//ensure(false);
 			ValueAsString = "<Unknown Data Type>";
@@ -312,6 +321,7 @@ void SCollectionSpreadSheet::Construct(const FArguments& InArgs)
 	AttrTypeWidthMap.Add("Box", 550);
 	AttrTypeWidthMap.Add("MeshSection", 0);
 	AttrTypeWidthMap.Add("UInt8", 0);
+	AttrTypeWidthMap.Add("Guid", 350);
 }
 
 
diff --git a/Engine/Plugins/Experimental/Fracture/Source/FractureEngine/Private/FractureEngineFracturing.cpp b/Engine/Plugins/Experimental/Fracture/Source/FractureEngine/Private/FractureEngineFracturing.cpp
--- a/Engine/Plugins/Experimental/Fracture/Source/FractureEngine/Private/FractureEngineFracturing.cpp
+++ b/Engine/Plugins/Experimental/Fracture/Source/FractureEngine/Private/FractureEngineFracturing.cpp
@@ -7,11 +7,11 @@
 #include "GeometryCollection/GeometryCollectionAlgo.h"
 #include "GeometryCollection/TransformCollection.h"
 #include "Voronoi/Voronoi.h"
-#include "PlanarCut.h"
 #include "GeometryCollection/Facades/CollectionMeshFacade.h"
 #include "Algo/RemoveIf.h"
 #include "GeometryCollection/Facades/CollectionTransformSelectionFacade.h"
 #include "GeometryCollection/GeometryCollectionClusteringUtility.h"
+#include "FractureEngineMaterials.h"
 
 static void AddAdditionalAttributesIfRequired(FManagedArrayCollection& InOutCollection)
 {
@@ -1063,5 +1063,250 @@ int32 FFractureEngineFracturing::MeshCutter(TArray<FTransform>& MeshTransforms,
 	return INDEX_NONE;
 }
 
+namespace UE::Dataflow::Private
+{
+	static void GenerateVoronoiSites(const FBox& InBoundingBox,
+		const int32 InMinVoronoiSites,
+		const int32 InMaxVoronoiSites,
+		const int32 InRandomSeed,
+		TArray<FVector>& OutSites)
+	{
+		FRandomStream RandStream(InRandomSeed);
+
+		const FVector Extent(InBoundingBox.Max - InBoundingBox.Min);
+
+		const int32 SiteCount = RandStream.RandRange(InMinVoronoiSites, InMaxVoronoiSites);
+
+		OutSites.Reserve(OutSites.Num() + SiteCount);
+		for (int32 ii = 0; ii < SiteCount; ++ii)
+		{
+			OutSites.Emplace(InBoundingBox.Min + FVector(RandStream.FRand(), RandStream.FRand(), RandStream.FRand()) * Extent);
+		}
+	}
+
+	static float GetMaxVertexMovement(const float InGrout,
+		const float InAmplitude,
+		const int32 InOctaveNumber,
+		const float InPersistence)
+	{
+		float MaxDisp = InGrout;
+		float AmplitudeScaled = InAmplitude;
+
+		for (int32 OctaveIdx = 0; OctaveIdx < InOctaveNumber; OctaveIdx++, AmplitudeScaled *= InPersistence)
+		{
+			MaxDisp += FMath::Abs(AmplitudeScaled);
+		}
+
+		return MaxDisp;
+	}
+
+	static FBox GetVoronoiBounds(const FBox& InBoundingBox,
+		const TArray<FVector>& Sites,
+		const float InGrout,
+		const float InAmplitude,
+		const int32 InOctaveNumber,
+		const float InPersistence)
+	{
+		FBox VoronoiBounds = InBoundingBox;
+		if (Sites.Num() > 0)
+		{
+			VoronoiBounds += FBox(Sites);
+		}
+
+		return VoronoiBounds.ExpandBy(GetMaxVertexMovement(InGrout, InAmplitude, InOctaveNumber, InPersistence) + KINDA_SMALL_NUMBER);
+	}
+
+	static void GenerateTemporaryGuids(FManagedArrayCollection& InCollection, int32 InStartIdx, bool InForceInit)
+	{
+		bool bNeedsInit = false;
+		if (!InCollection.HasAttribute("GUID", FTransformCollection::TransformGroup))
+		{
+			FManagedArrayCollection::FConstructionParameters Params(FName(""), false);
+			InCollection.AddAttribute<FGuid>("GUID", FTransformCollection::TransformGroup, Params);
+			bNeedsInit = true;
+		}
+
+		if (bNeedsInit || InForceInit)
+		{
+			TManagedArray<FGuid>& Guids = InCollection.ModifyAttribute<FGuid>("GUID", FTransformCollection::TransformGroup);
+			for (int32 Idx = InStartIdx; Idx < Guids.Num(); ++Idx)
+			{
+				Guids[Idx] = FGuid::NewGuid();
+			}
+		}
+	}
+
+	static int32 UniformFractureProc(
+		TUniquePtr<FGeometryCollection>& OutGeomCollection,
+		const TArray<int32>& InTransformSelectionArr,
+		const FUniformFractureProcSettings& InUniformFractureProcSettings)
+	{
+		TArray<FVector> Sites;
+		UE::Dataflow::Private::GenerateVoronoiSites(InUniformFractureProcSettings.BBox,
+			InUniformFractureProcSettings.MinVoronoiSites,
+			InUniformFractureProcSettings.MaxVoronoiSites,
+			InUniformFractureProcSettings.RandomSeed,
+			Sites);
+
+		FBox VoronoiBounds = UE::Dataflow::Private::GetVoronoiBounds(InUniformFractureProcSettings.BBox,
+			Sites,
+			InUniformFractureProcSettings.Grout,
+			InUniformFractureProcSettings.NoiseSettings.Amplitude,
+			InUniformFractureProcSettings.NoiseSettings.Octaves,
+			InUniformFractureProcSettings.NoiseSettings.Persistence);
+
+		FVector Origin = InUniformFractureProcSettings.Transform.GetTranslation();
+		for (FVector& Site : Sites)
+		{
+			Site -= Origin;
+		}
+		VoronoiBounds.Min -= Origin;
+		VoronoiBounds.Max -= Origin;
+		FVoronoiDiagram Voronoi(Sites, VoronoiBounds, .1f);
+
+		FPlanarCells VoronoiPlanarCells = FPlanarCells(Sites, Voronoi);
+		VoronoiPlanarCells.InternalSurfaceMaterials.NoiseSettings = InUniformFractureProcSettings.NoiseSettings;
+
+		int32 FirstNewGeometryIndex = CutMultipleWithPlanarCells(VoronoiPlanarCells, 
+			*OutGeomCollection, 
+			InTransformSelectionArr, 
+			InUniformFractureProcSettings.Grout,
+			InUniformFractureProcSettings.CollisionSampleSpacing,
+			InUniformFractureProcSettings.RandomSeed, 
+			InUniformFractureProcSettings.Transform, 
+			true, 
+			true, 
+			nullptr, 
+			Origin, 
+			InUniformFractureProcSettings.SplitIslands);
+
+		// Assign internal material
+		if (InUniformFractureProcSettings.InternalMaterialID > INDEX_NONE)
+		{
+			FFractureEngineMaterials::SetMaterialOnGeometryAfter(*OutGeomCollection, FirstNewGeometryIndex, FFractureEngineMaterials::ETargetFaces::InternalFaces, InUniformFractureProcSettings.InternalMaterialID);
+			OutGeomCollection->ReindexMaterials();
+		}
+
+		// Generate GUIDs
+		GenerateTemporaryGuids(*OutGeomCollection, FirstNewGeometryIndex, true);
+
+		return FirstNewGeometryIndex;
+	}
+}
 
+int32 FFractureEngineFracturing::UniformFracture(
+	FManagedArrayCollection& InOutCollection,
+	FDataflowTransformSelection InTransformSelection, 
+	const FUniformFractureSettings& InUniformFractureSettings)
+{
+	if (TUniquePtr<FGeometryCollection> GeomCollection = TUniquePtr<FGeometryCollection>(InOutCollection.NewCopy<FGeometryCollection>()))
+	{
+		RandomReduceSelection(InTransformSelection, InUniformFractureSettings.RandomSeed, InUniformFractureSettings.ChanceToFracture);
+		TArray<int32> TransformSelectionArr = InTransformSelection.AsArray();
+
+		if (!FFractureEngineSelection::IsBoneSelectionValid(InOutCollection, TransformSelectionArr))
+		{
+			return INDEX_NONE;
+		}
+
+		// Update global transforms and bounds		
+		const TManagedArray<FTransform3f>& Transform = GeomCollection->GetAttribute<FTransform3f>("Transform", FGeometryCollection::TransformGroup);
+		const TManagedArray<int32>& TransformToGeometryIndex = GeomCollection->GetAttribute<int32>("TransformToGeometryIndex", FGeometryCollection::TransformGroup);
+		const TManagedArray<FBox>& BoundingBoxes = GeomCollection->GetAttribute<FBox>("BoundingBox", FGeometryCollection::GeometryGroup);
+
+		TArray<FTransform> Transforms;
+		GeometryCollectionAlgo::GlobalMatrices(Transform, GeomCollection->Parent, Transforms);
+
+		int32 TransformCount = Transform.Num();
+		TMap<int32, FBox> BoundsToBone;
+		for (int32 Index = 0; Index < TransformCount; ++Index)
+		{
+			if (TransformToGeometryIndex[Index] > INDEX_NONE)
+			{
+				BoundsToBone.Add(Index, BoundingBoxes[TransformToGeometryIndex[Index]].TransformBy(Transforms[Index]));
+			}
+		}
+
+		FBox BoundingBox(ForceInit);
+
+		if (InUniformFractureSettings.GroupFracture)
+		{
+			FBox Bounds(ForceInit);
+			for (int32 TransformIndex : TransformSelectionArr)
+			{
+				if (TransformToGeometryIndex[TransformIndex] > INDEX_NONE)
+				{
+					Bounds += BoundsToBone[TransformIndex];
+				}
+			}
+			
+			BoundingBox = Bounds;
+
+			// Fracture
+			FUniformFractureProcSettings UniformFractureProcSettings;
+			UniformFractureProcSettings.BBox = BoundingBox;
+			UniformFractureProcSettings.Transform = InUniformFractureSettings.Transform;
+			UniformFractureProcSettings.MinVoronoiSites = InUniformFractureSettings.MinVoronoiSites;
+			UniformFractureProcSettings.MaxVoronoiSites = InUniformFractureSettings.MaxVoronoiSites;
+			UniformFractureProcSettings.InternalMaterialID = InUniformFractureSettings.InternalMaterialID;
+			UniformFractureProcSettings.RandomSeed = InUniformFractureSettings.RandomSeed;
+			UniformFractureProcSettings.SplitIslands = InUniformFractureSettings.SplitIslands;
+			UniformFractureProcSettings.Grout = InUniformFractureSettings.Grout;
+			UniformFractureProcSettings.NoiseSettings = InUniformFractureSettings.NoiseSettings;
+			UniformFractureProcSettings.AddSamplesForCollision = InUniformFractureSettings.AddSamplesForCollision;
+			UniformFractureProcSettings.CollisionSampleSpacing = InUniformFractureSettings.CollisionSampleSpacing;
+
+			int ResultGeometryIndex = UE::Dataflow::Private::UniformFractureProc(
+				GeomCollection,
+				TransformSelectionArr,
+				UniformFractureProcSettings);
+
+			InOutCollection = (const FManagedArrayCollection&)(*GeomCollection);
+
+			return ResultGeometryIndex;
+		}
+		else
+		{
+			int ResultGeometryIndex = INDEX_NONE;
+
+			for (int32 TransformIndex : TransformSelectionArr)
+			{
+				TArray<int32> TransformSelection;
+				int32 Seed;
+
+				if (TransformToGeometryIndex[TransformIndex] > INDEX_NONE)
+				{
+					TransformSelection.Add(TransformIndex);
+					Seed = InUniformFractureSettings.RandomSeed + TransformIndex;
+					BoundingBox = BoundsToBone[TransformIndex];
+
+					// Fracture
+					FUniformFractureProcSettings UniformFractureProcSettings;
+					UniformFractureProcSettings.BBox = BoundingBox;
+					UniformFractureProcSettings.Transform = InUniformFractureSettings.Transform;
+					UniformFractureProcSettings.MinVoronoiSites = InUniformFractureSettings.MinVoronoiSites;
+					UniformFractureProcSettings.MaxVoronoiSites = InUniformFractureSettings.MaxVoronoiSites;
+					UniformFractureProcSettings.InternalMaterialID = InUniformFractureSettings.InternalMaterialID;
+					UniformFractureProcSettings.RandomSeed = Seed;
+					UniformFractureProcSettings.SplitIslands = InUniformFractureSettings.SplitIslands;
+					UniformFractureProcSettings.Grout = InUniformFractureSettings.Grout;
+					UniformFractureProcSettings.NoiseSettings = InUniformFractureSettings.NoiseSettings;
+					UniformFractureProcSettings.AddSamplesForCollision = InUniformFractureSettings.AddSamplesForCollision;
+					UniformFractureProcSettings.CollisionSampleSpacing = InUniformFractureSettings.CollisionSampleSpacing;
+
+					ResultGeometryIndex = UE::Dataflow::Private::UniformFractureProc(
+						GeomCollection,
+						TransformSelection,
+						UniformFractureProcSettings);
+				}
+			}
+
+			InOutCollection = (const FManagedArrayCollection&)(*GeomCollection);
+
+			return ResultGeometryIndex;
+		}
+	}
+
+	return INDEX_NONE;
+}
 
diff --git a/Engine/Plugins/Experimental/Fracture/Source/FractureEngine/Public/FractureEngineFracturing.h b/Engine/Plugins/Experimental/Fracture/Source/FractureEngine/Public/FractureEngineFracturing.h
--- a/Engine/Plugins/Experimental/Fracture/Source/FractureEngine/Public/FractureEngineFracturing.h
+++ b/Engine/Plugins/Experimental/Fracture/Source/FractureEngine/Public/FractureEngineFracturing.h
@@ -3,13 +3,15 @@
 #pragma once
 
 #include "Math/Box.h"
+#include "PlanarCut.h"
+#include "Dataflow/DataflowSelection.h"
 
 #include "FractureEngineFracturing.generated.h"
 
 namespace UE { namespace Geometry { class FDynamicMesh3; } }
 
-struct FDataflowTransformSelection;
 struct FManagedArrayCollection;
+class FGeometryCollection;
 
 UENUM(BlueprintType)
 enum class EFractureBrickBondEnum : uint8
@@ -32,6 +34,37 @@ enum class EMeshCutterCutDistribution : uint8
 	Grid UMETA(DisplayName = "Grid"),
 };
 
+struct FUniformFractureSettings
+{
+	FTransform Transform;
+	int32 MinVoronoiSites;
+	int32 MaxVoronoiSites;
+	int32 InternalMaterialID;
+	int32 RandomSeed;
+	float ChanceToFracture;
+	bool GroupFracture;
+	bool SplitIslands;
+	float Grout;
+	FNoiseSettings NoiseSettings;
+	bool AddSamplesForCollision;
+	float CollisionSampleSpacing;
+};
+
+struct FUniformFractureProcSettings
+{
+	FBox BBox;
+	FTransform Transform;
+	int32 MinVoronoiSites;
+	int32 MaxVoronoiSites;
+	int32 InternalMaterialID;
+	int32 RandomSeed;
+	bool SplitIslands;
+	float Grout;
+	FNoiseSettings NoiseSettings;
+	bool AddSamplesForCollision;
+	float CollisionSampleSpacing;
+};
+
 class FRACTUREENGINE_API FFractureEngineFracturing
 {
 public:
@@ -161,4 +194,11 @@ public:
 		const float InChanceToFracture,
 		const bool InSplitIslands,
 		const float InCollisionSampleSpacing);
+
+	static int32 UniformFracture(
+		FManagedArrayCollection& InOutCollection,
+		FDataflowTransformSelection InTransformSelection,
+		const FUniformFractureSettings& InUniformFractureSettings);
+
 };
+	
\ No newline at end of file
diff --git a/Engine/Plugins/Experimental/GeometryCollectionPlugin/Source/GeometryCollectionNodes/Private/Dataflow/GeometryCollectionFracturingNodes.cpp b/Engine/Plugins/Experimental/GeometryCollectionPlugin/Source/GeometryCollectionNodes/Private/Dataflow/GeometryCollectionFracturingNodes.cpp
--- a/Engine/Plugins/Experimental/GeometryCollectionPlugin/Source/GeometryCollectionNodes/Private/Dataflow/GeometryCollectionFracturingNodes.cpp
+++ b/Engine/Plugins/Experimental/GeometryCollectionPlugin/Source/GeometryCollectionNodes/Private/Dataflow/GeometryCollectionFracturingNodes.cpp
@@ -59,6 +59,7 @@ namespace UE::Dataflow
 		DATAFLOW_NODE_REGISTER_CREATION_FACTORY(FSliceCutterDataflowNode);
 		DATAFLOW_NODE_REGISTER_CREATION_FACTORY(FBrickCutterDataflowNode);
 		DATAFLOW_NODE_REGISTER_CREATION_FACTORY(FMeshCutterDataflowNode);
+		DATAFLOW_NODE_REGISTER_CREATION_FACTORY(FUniformFractureDataflowNode);
 
 		// GeometryCollection|Fracture
 		DATAFLOW_NODE_REGISTER_CREATION_FACTORY_NODE_COLORS_BY_CATEGORY("GeometryCollection|Fracture", FLinearColor(1.f, 1.f, 0.8f), CDefaultNodeBodyTintColor);
@@ -953,6 +954,148 @@ void FMeshCutterDataflowNode::Evaluate(UE::Dataflow::FContext& Context, const FD
 	}
 }
 
+FUniformFractureDataflowNode::FUniformFractureDataflowNode(const UE::Dataflow::FNodeParameters& InParam, FGuid InGuid)
+	: FDataflowNode(InParam, InGuid)
+{
+	RegisterInputConnection(&Collection);
+	RegisterInputConnection(&TransformSelection);
+	RegisterInputConnection(&Transform);
+	RegisterInputConnection(&MinVoronoiSites)
+		.SetCanHidePin(true)
+		.SetPinIsHidden(true);
+	RegisterInputConnection(&MaxVoronoiSites)
+		.SetCanHidePin(true)
+		.SetPinIsHidden(true);
+	RegisterInputConnection(&InternalMaterialID);
+	RegisterInputConnection(&RandomSeed)
+		.SetCanHidePin(true)
+		.SetPinIsHidden(true);
+	RegisterInputConnection(&ChanceToFracture)
+		.SetCanHidePin(true)
+		.SetPinIsHidden(true);
+	RegisterInputConnection(&Grout)
+		.SetCanHidePin(true)
+		.SetPinIsHidden(true);
+	RegisterInputConnection(&Amplitude)
+		.SetCanHidePin(true)
+		.SetPinIsHidden(true);
+	RegisterInputConnection(&Frequency)
+		.SetCanHidePin(true)
+		.SetPinIsHidden(true);
+	RegisterInputConnection(&Persistence)
+		.SetCanHidePin(true)
+		.SetPinIsHidden(true);
+	RegisterInputConnection(&Lacunarity)
+		.SetCanHidePin(true)
+		.SetPinIsHidden(true);
+	RegisterInputConnection(&OctaveNumber)
+		.SetCanHidePin(true)
+		.SetPinIsHidden(true);
+	RegisterInputConnection(&PointSpacing)
+		.SetCanHidePin(true)
+		.SetPinIsHidden(true);
+	RegisterInputConnection(&CollisionSampleSpacing)
+		.SetCanHidePin(true)
+		.SetPinIsHidden(true);
+	RegisterOutputConnection(&Collection, &Collection);
+	RegisterOutputConnection(&TransformSelection, &TransformSelection);
+	RegisterOutputConnection(&NewGeometryTransformSelection);
+}
+
+void FUniformFractureDataflowNode::Evaluate(UE::Dataflow::FContext& Context, const FDataflowOutput* Out) const
+{
+	if (Out->IsA(&Collection) ||
+		Out->IsA(&TransformSelection) ||
+		Out->IsA(&NewGeometryTransformSelection))
+	{
+		FDataflowTransformSelection InTransformSelection = GetValue(Context, &TransformSelection);
+
+		//
+		// If not connected select everything by default
+		//
+		if (!IsConnected(&TransformSelection))
+		{
+			const FManagedArrayCollection& InCollection = GetValue<FManagedArrayCollection>(Context, &Collection);
+
+			GeometryCollection::Facades::FCollectionTransformSelectionFacade TransformSelectionFacade(InCollection);
+			const TArray<int32>& SelectionArr = TransformSelectionFacade.SelectAll();
+
+			FDataflowTransformSelection NewTransformSelection;
+			NewTransformSelection.Initialize(InCollection.NumElements(FGeometryCollection::TransformGroup), false);
+			NewTransformSelection.SetFromArray(SelectionArr);
+
+			InTransformSelection = NewTransformSelection;
+		}
+
+		if (InTransformSelection.AnySelected())
+		{
+			FManagedArrayCollection InCollection = GetValue<FManagedArrayCollection>(Context, &Collection);
+
+			FUniformFractureSettings UniformFractureSettings;
+			UniformFractureSettings.Transform = GetValue(Context, &Transform);
+			UniformFractureSettings.MinVoronoiSites = GetValue(Context, &MinVoronoiSites);
+			UniformFractureSettings.MaxVoronoiSites = GetValue(Context, &MaxVoronoiSites);
+			UniformFractureSettings.InternalMaterialID = InternalMaterialID;
+			UniformFractureSettings.RandomSeed = GetValue(Context, &RandomSeed);
+			UniformFractureSettings.ChanceToFracture = GetValue(Context, &ChanceToFracture);
+			UniformFractureSettings.GroupFracture = GroupFracture;
+			UniformFractureSettings.SplitIslands = SplitIslands;
+			UniformFractureSettings.Grout = GetValue(Context, &Grout);
+			UniformFractureSettings.NoiseSettings.Amplitude = GetValue(Context, &Amplitude);
+			UniformFractureSettings.NoiseSettings.Frequency = GetValue(Context, &Frequency);
+			UniformFractureSettings.NoiseSettings.Persistence = GetValue(Context, &Persistence);
+			UniformFractureSettings.NoiseSettings.Lacunarity = GetValue(Context, &Lacunarity);
+			UniformFractureSettings.NoiseSettings.Octaves = GetValue(Context, &OctaveNumber);
+			UniformFractureSettings.NoiseSettings.PointSpacing = GetValue(Context, &PointSpacing);
+			UniformFractureSettings.AddSamplesForCollision = AddSamplesForCollision;
+			UniformFractureSettings.CollisionSampleSpacing = GetValue(Context, &CollisionSampleSpacing);
+
+			int32 ResultGeometryIndex = FFractureEngineFracturing::UniformFracture(
+				InCollection,
+				InTransformSelection,
+				UniformFractureSettings);
+
+			FDataflowTransformSelection NewSelection;
+			FDataflowTransformSelection OriginalSelection;
+
+			if (ResultGeometryIndex != INDEX_NONE)
+			{
+				if (InCollection.HasAttribute("TransformIndex", FGeometryCollection::GeometryGroup))
+				{
+					const TManagedArray<int32>& GeometryToTransformIndices = InCollection.GetAttribute<int32>("TransformIndex", FGeometryCollection::GeometryGroup);
+
+					const int32 NumTransforms = InCollection.NumElements(FGeometryCollection::TransformGroup);
+					NewSelection.Initialize(NumTransforms, false);
+					OriginalSelection.Initialize(NumTransforms, false);
 
+					// The newly fractured pieces are added to the end of the transform array (starting position is ResultGeometryIndex)
+					for (int32 GeometryIdx = ResultGeometryIndex; GeometryIdx < GeometryToTransformIndices.Num(); ++GeometryIdx)
+					{
+						const int32 TransformIdx = GeometryToTransformIndices[GeometryIdx];
+						NewSelection.SetSelected(TransformIdx);
+					}
+
+					for (int32 TransformIdx = 0; TransformIdx < InTransformSelection.Num(); ++TransformIdx)
+					{
+						if (InTransformSelection.IsSelected(TransformIdx))
+						{
+							OriginalSelection.SetSelected(TransformIdx);
+						}
+					}
+				}
+			}
+
+			SetValue(Context, MoveTemp(InCollection), &Collection);
+			SetValue(Context, OriginalSelection, &TransformSelection);
+			SetValue(Context, NewSelection, &NewGeometryTransformSelection);
+
+			return;
+		}
+
+		SafeForwardInput(Context, &Collection, &Collection);
+		SetValue(Context, InTransformSelection, &TransformSelection);
+		SetValue(Context, FDataflowTransformSelection(), &NewGeometryTransformSelection);
+	}
+}
 
 
diff --git a/Engine/Plugins/Experimental/GeometryCollectionPlugin/Source/GeometryCollectionNodes/Public/Dataflow/GeometryCollectionFracturingNodes.h b/Engine/Plugins/Experimental/GeometryCollectionPlugin/Source/GeometryCollectionNodes/Public/Dataflow/GeometryCollectionFracturingNodes.h
--- a/Engine/Plugins/Experimental/GeometryCollectionPlugin/Source/GeometryCollectionNodes/Public/Dataflow/GeometryCollectionFracturingNodes.h
+++ b/Engine/Plugins/Experimental/GeometryCollectionPlugin/Source/GeometryCollectionNodes/Public/Dataflow/GeometryCollectionFracturingNodes.h
@@ -305,7 +305,7 @@ public:
 
 /**
  *
- * Editor Fracture Mode / Fracture / Uniform tool
+ * Voronoi fracture
  * Fracture using a Voronoi diagram with a uniform random pattern, creating fracture pieces of similar volume across the shape.
  *
  */
@@ -387,7 +387,7 @@ public:
 
 /**
  *
- * Editor Fracture Mode / Fracture / Uniform tool
+ * Voronoi fracture
  * Fracture using a Voronoi diagram with a uniform random pattern, creating fracture pieces of similar volume across the shape.
  *
  */
@@ -1167,6 +1167,117 @@ public:
 	virtual void Evaluate(UE::Dataflow::FContext& Context, const FDataflowOutput* Out) const override;
 };
 
+/**
+ *
+ * Editor Fracture Mode / Fracture / Uniform tool
+ * Fracture using a Voronoi diagram with a uniform random pattern, creating fracture pieces of similar volume across the shape.
+ *
+ */
+USTRUCT(meta = (DataflowGeometryCollection))
+struct FUniformFractureDataflowNode : public FDataflowNode
+{
+	GENERATED_USTRUCT_BODY()
+	DATAFLOW_NODE_DEFINE_INTERNAL(FUniformFractureDataflowNode, "UniformFracture", "GeometryCollection|Fracture", "")
+	DATAFLOW_NODE_RENDER_TYPE("SurfaceRender", FGeometryCollection::StaticType(), "Collection")
+
+private:
+	/** Collection to fracture */
+	UPROPERTY(meta = (DataflowInput, DataflowOutput, DataflowPassthrough = "Collection", DataflowIntrinsic))
+	FManagedArrayCollection Collection;
+
+	/** Bones to fracture, if not connected it will fracture all the bones */
+	UPROPERTY(meta = (DataflowInput, DataflowOutput, DataflowPassthrough = "TransformSelection", DisplayName = "TransformSelection"))
+	FDataflowTransformSelection TransformSelection;
+
+	/** Transform to apply */
+	UPROPERTY(meta = (DataflowInput))
+	FTransform Transform = FTransform::Identity;
+
+	/** Minimum Number of Voronoi sites. The amount of sites per Voronoi diagram will be chosen at random between Min and Max */
+	UPROPERTY(EditAnywhere, Category = "Uniform Voronoi", meta = (DataflowInput, UIMin = "1", UIMax = "5000", ClampMin = "1"));
+	int32 MinVoronoiSites = 20;
+
+	/** Maximum Number of Voronoi sites. The amount of sites per Voronoi diagram will be chosen at random between Min and Max */
+	UPROPERTY(EditAnywhere, Category = "Uniform Voronoi", meta = (DataflowInput, UIMin = "1", UIMax = "5000", ClampMin = "1"));
+	int32 MaxVoronoiSites = 20;
+
+	/** ID for the material for the newly created internal faces */
+	UPROPERTY(EditAnywhere, Category = "Materials", meta = (DataflowInput, UIMin = "0", UIMax = "10", ClampMin = "0"));
+	int32 InternalMaterialID = 0;
+
+	/** Random number generator seed for repeatability. If the value is -1, a different random seed will be used every time, otherwise the specified seed will always be used */
+	UPROPERTY(EditAnywhere, Category = "Common Fracture", meta = (DataflowInput, UIMin = "-1", UIMax = "5000", ClampMin = "-1"));
+	int32 RandomSeed = 0;
+
+	/** Chance to fracture each selected bone. If 0, no bones will fracture; if 1, all bones will fracture. */
+	UPROPERTY(EditAnywhere, Category = "Common Fracture", meta = (DataflowInput, UIMin = "0.0", UIMax = "1.0", ClampMin = "0.0", ClampMax = "1.0"));
+	float ChanceToFracture = 1.f;
+
+	/** Generate a fracture pattern across all selected meshes.  */
+	UPROPERTY(EditAnywhere, Category = "Common Fracture");
+	bool GroupFracture = true;
+
+	/** Whether to split the fractured mesh pieces based on geometric connectivity after fracturing */
+	UPROPERTY(EditAnywhere, Category = "Common Fracture");
+	bool SplitIslands = true;
+
+	/** Amount of space to leave between cut pieces */
+	UPROPERTY(EditAnywhere, Category = "Common Fracture", meta = (DataflowInput, UIMin = 0.f));
+	float Grout = 0.f;
+
+	/** Size of the Perlin noise displacement (in cm). If 0, no noise will be applied */
+	UPROPERTY(EditAnywhere, Category = "Noise", meta = (DataflowInput, UIMin = 0.f));
+	float Amplitude = 0.f;
+
+	/** Period of the Perlin noise.  Smaller values will create a smoother noise pattern */
+	UPROPERTY(EditAnywhere, Category = "Noise", meta = (DataflowInput, UIMin = 0.00001f));
+	float Frequency = 0.1f;
+
+	/** Persistence of the layers of Perlin noise. At each layer (octave) after the first, the amplitude of the Perlin noise is scaled by this factor */
+	UPROPERTY(EditAnywhere, Category = "Noise", meta = (DataflowInput, UIMin = 0.f));
+	float Persistence = 0.5f;
+
+	/** Lacunarity of the layers of Perlin noise. At each layer (octave) after the first, the frequency of the Perlin noise is scaled by this factor */
+	UPROPERTY(EditAnywhere, Category = "Noise", meta = (DataflowInput, UIMin = 0.f));
+	float Lacunarity = 2.f;
+
+	/**
+	 * Number of fractal layers of Perlin noise to apply. Each layer is additive, with Amplitude and Frequency parameters scaled by Persistence and Lacunarity
+	 * Smaller values (1 or 2) will create noise that looks like gentle rolling hills, while larger values (> 4) will tend to look more like craggy mountains
+	 */
+	UPROPERTY(EditAnywhere, Category = "Noise", meta = (DataflowInput, UIMin = 0, ClampMin = 0));
+	int32 OctaveNumber = 4;
+
+	/** Distance (in cm) between vertices on cut surfaces where noise is added.  Larger spacing between vertices will create more efficient meshes with fewer triangles, but less resolution to see the shape of the added noise  */
+	UPROPERTY(EditAnywhere, Category = "Noise", meta = (DataflowInput, UIMin = 0.f));
+	float PointSpacing = 10.f;
+
+	/**
+	 * If enabled, add extra vertices (without triangles) to the geometry in regions where vertices are spaced too far apart (e.g. across large triangles)
+	 * These extra vertices will be used as collision samples in particle-implicit collisions, and can help the physics system detect collisions more accurately
+	 *
+	 * Note this is *only* useful for simulations that use particle-implicit collisions
+	 */
+	UPROPERTY(EditAnywhere, Category = "Collision");
+	bool AddSamplesForCollision = false;
+
+	/**
+	 * The number of centimeters to allow between vertices on the mesh surface: If there are gaps larger than this, add additional vertices (without triangles) to help support particle-implicit collisions
+	 * Only used if Add Samples For Collision is enabled
+	 */
+	UPROPERTY(EditAnywhere, Category = "Collision", meta = (DataflowInput, UIMin = 0.f, EditCondition = "AddSamplesForCollision"));
+	float CollisionSampleSpacing = 50.f;
+
+	/** Fractured Pieces */
+	UPROPERTY(meta = (DataflowOutput, DisplayName = "FracturedTransformSelection "))
+	FDataflowTransformSelection NewGeometryTransformSelection;
+
+public:
+	FUniformFractureDataflowNode(const UE::Dataflow::FNodeParameters& InParam, FGuid InGuid = FGuid::NewGuid());
+
+	virtual void Evaluate(UE::Dataflow::FContext& Context, const FDataflowOutput* Out) const override;
+};
+
 namespace UE::Dataflow
 {
 	void GeometryCollectionFracturingNodes();
