Dataflow:
- Added FVisualizeFractureDataflowNode
- Added Offset to FExplodedViewDataflowNode
- Bug fixes

diff --git a/Engine/Plugins/Experimental/Fracture/Source/FractureEngine/Private/FractureEngineFracturing.cpp b/Engine/Plugins/Experimental/Fracture/Source/FractureEngine/Private/FractureEngineFracturing.cpp
--- a/Engine/Plugins/Experimental/Fracture/Source/FractureEngine/Private/FractureEngineFracturing.cpp
+++ b/Engine/Plugins/Experimental/Fracture/Source/FractureEngine/Private/FractureEngineFracturing.cpp
@@ -84,29 +84,28 @@ static bool GetValidGeoCenter(FManagedArrayCollection& InOutCollection,
 //
 // TODO: Rewrite this using facades
 //
-void FFractureEngineFracturing::GenerateExplodedViewAttribute(FManagedArrayCollection& InOutCollection, const FVector& InScale, float InUniformScale)
+void FFractureEngineFracturing::GenerateExplodedViewAttribute(FManagedArrayCollection& InOutCollection, const FVector& InScale, const float InUniformScale, const int32 InViewFractureLevel, const int32 InMaxFractureLevel)
 {
 	// Check if InOutCollection is not empty
-	if (InOutCollection.HasAttribute("Transform", FGeometryCollection::TransformGroup))
+	if (InOutCollection.HasAttribute(FTransformCollection::TransformAttribute, FGeometryCollection::TransformGroup))
 	{
 		InOutCollection.AddAttribute<FVector3f>("ExplodedVector", FGeometryCollection::TransformGroup, FManagedArrayCollection::FConstructionParameters(FName(), false));
 		check(InOutCollection.HasAttribute("ExplodedVector", FGeometryCollection::TransformGroup));
 
 		TManagedArray<FVector3f>& ExplodedVectors = InOutCollection.ModifyAttribute<FVector3f>("ExplodedVector", FGeometryCollection::TransformGroup);
-		const TManagedArray<FTransform3f>& Transforms = InOutCollection.GetAttribute<FTransform3f>("Transform", FGeometryCollection::TransformGroup);
-		const TManagedArray<int32>& TransformToGeometryIndices = InOutCollection.GetAttribute<int32>("TransformToGeometryIndex", FGeometryCollection::TransformGroup);
-		const TManagedArray<FBox>& BoundingBoxes = InOutCollection.GetAttribute<FBox>("BoundingBox", FGeometryCollection::GeometryGroup);
+		const TManagedArray<FTransform3f>& Transforms = InOutCollection.GetAttribute<FTransform3f>(FTransformCollection::TransformAttribute, FGeometryCollection::TransformGroup);
+		const TManagedArray<int32>& TransformToGeometryIndices = InOutCollection.GetAttribute<int32>(FGeometryCollection::TransformToGeometryIndexAttribute, FGeometryCollection::TransformGroup);
+		const TManagedArray<FBox>& BoundingBoxes = InOutCollection.GetAttribute<FBox>(FGeometryCollection::BoundingBoxAttribute, FGeometryCollection::GeometryGroup);
 
 		// Make sure we have valid "Level"
 		AddAdditionalAttributesIfRequired(InOutCollection);
 
-		const TManagedArray<int32>& Levels = InOutCollection.GetAttribute<int32>("Level", FTransformCollection::TransformGroup);
-		const TManagedArray<int32>& Parents = InOutCollection.GetAttribute<int32>("Parent", FTransformCollection::TransformGroup);
-		const TManagedArray<TSet<int32>>& Children = InOutCollection.GetAttribute<TSet<int32>>("Children", FGeometryCollection::TransformGroup);
-		const TManagedArray<int32>& SimulationTypes = InOutCollection.GetAttribute<int32>("SimulationType", FGeometryCollection::TransformGroup);
+		const TManagedArray<int32>& Levels = InOutCollection.GetAttribute<int32>(FTransformCollection::LevelAttribute, FTransformCollection::TransformGroup);
+		const TManagedArray<int32>& Parents = InOutCollection.GetAttribute<int32>(FTransformCollection::ParentAttribute, FTransformCollection::TransformGroup);
+		const TManagedArray<TSet<int32>>& Children = InOutCollection.GetAttribute<TSet<int32>>(FTransformCollection::ChildrenAttribute, FGeometryCollection::TransformGroup);
+		const TManagedArray<int32>& SimulationTypes = InOutCollection.GetAttribute<int32>(FGeometryCollection::SimulationTypeAttribute, FGeometryCollection::TransformGroup);
 
-		int32 ViewFractureLevel = -1;
-		int32 MaxFractureLevel = ViewFractureLevel;
+		int32 MaxFractureLevel = InMaxFractureLevel;
 		for (int32 Idx = 0, ni = Transforms.Num(); Idx < ni; ++Idx)
 		{
 			if (Levels[Idx] > MaxFractureLevel)
@@ -130,7 +129,7 @@ void FFractureEngineFracturing::GenerateExplodedViewAttribute(FManagedArrayColle
 			if (GetValidGeoCenter(InOutCollection, TransformToGeometryIndices, TransformArr, Parents, Children, BoundingBoxes, SimulationTypes, Idx, GeoCenter))
 			{
 				TransformedCenters[Idx] = GeoCenter;
-				if ((ViewFractureLevel < 0) || Levels[Idx] == ViewFractureLevel)
+				if ((InViewFractureLevel < 0) || Levels[Idx] == InViewFractureLevel)
 				{
 					Center += TransformedCenters[Idx];
 					++TransformsCount;
@@ -144,7 +143,7 @@ void FFractureEngineFracturing::GenerateExplodedViewAttribute(FManagedArrayColle
 		{
 			for (int32 Idx = 0, ni = TransformArr.Num(); Idx < ni; ++Idx)
 			{
-				if ((ViewFractureLevel < 0) || Levels[Idx] == ViewFractureLevel)
+				if ((InViewFractureLevel < 0) || Levels[Idx] == InViewFractureLevel)
 				{
 					FVector ScaleVec = InScale * InUniformScale;
 					ExplodedVectors[Idx] = (FVector3f)(TransformedCenters[Idx] - Center) * (FVector3f)ScaleVec;
@@ -222,15 +221,15 @@ int32 FFractureEngineFracturing::VoronoiFracture(FManagedArrayCollection& InOutC
 			//
 			FBox BoundingBox(ForceInit);
 
-			if (InOutCollection.HasAttribute("Transform", FGeometryCollection::TransformGroup) &&
-				InOutCollection.HasAttribute("Parent", FGeometryCollection::TransformGroup) &&
-				InOutCollection.HasAttribute("TransformIndex", FGeometryCollection::GeometryGroup) &&
-				InOutCollection.HasAttribute("BoundingBox", FGeometryCollection::GeometryGroup))
+			if (InOutCollection.HasAttribute(FTransformCollection::TransformAttribute, FGeometryCollection::TransformGroup) &&
+				InOutCollection.HasAttribute(FTransformCollection::ParentAttribute, FGeometryCollection::TransformGroup) &&
+				InOutCollection.HasAttribute(FGeometryCollection::TransformIndexAttribute, FGeometryCollection::GeometryGroup) &&
+				InOutCollection.HasAttribute(FGeometryCollection::BoundingBoxAttribute, FGeometryCollection::GeometryGroup))
 			{
-				const TManagedArray<FTransform3f>& Transforms = InOutCollection.GetAttribute<FTransform3f>("Transform", FGeometryCollection::TransformGroup);
-				const TManagedArray<int32>& ParentIndices = InOutCollection.GetAttribute<int32>("Parent", FGeometryCollection::TransformGroup);
-				const TManagedArray<int32>& TransformIndices = InOutCollection.GetAttribute<int32>("TransformIndex", FGeometryCollection::GeometryGroup);
-				const TManagedArray<FBox>& BoundingBoxes = InOutCollection.GetAttribute<FBox>("BoundingBox", FGeometryCollection::GeometryGroup);
+				const TManagedArray<FTransform3f>& Transforms = InOutCollection.GetAttribute<FTransform3f>(FTransformCollection::TransformAttribute, FGeometryCollection::TransformGroup);
+				const TManagedArray<int32>& ParentIndices = InOutCollection.GetAttribute<int32>(FTransformCollection::ParentAttribute, FGeometryCollection::TransformGroup);
+				const TManagedArray<int32>& TransformIndices = InOutCollection.GetAttribute<int32>(FGeometryCollection::TransformIndexAttribute, FGeometryCollection::GeometryGroup);
+				const TManagedArray<FBox>& BoundingBoxes = InOutCollection.GetAttribute<FBox>(FGeometryCollection::BoundingBoxAttribute, FGeometryCollection::GeometryGroup);
 
 				TArray<FMatrix> TmpGlobalMatrices;
 				GeometryCollectionAlgo::GlobalMatrices(Transforms, ParentIndices, TmpGlobalMatrices);
@@ -1210,9 +1209,9 @@ int32 FFractureEngineFracturing::UniformFracture(
 		}
 
 		// Update global transforms and bounds		
-		const TManagedArray<FTransform3f>& Transform = GeomCollection->GetAttribute<FTransform3f>("Transform", FGeometryCollection::TransformGroup);
-		const TManagedArray<int32>& TransformToGeometryIndex = GeomCollection->GetAttribute<int32>("TransformToGeometryIndex", FGeometryCollection::TransformGroup);
-		const TManagedArray<FBox>& BoundingBoxes = GeomCollection->GetAttribute<FBox>("BoundingBox", FGeometryCollection::GeometryGroup);
+		const TManagedArray<FTransform3f>& Transform = GeomCollection->GetAttribute<FTransform3f>(FTransformCollection::TransformAttribute, FGeometryCollection::TransformGroup);
+		const TManagedArray<int32>& TransformToGeometryIndex = GeomCollection->GetAttribute<int32>(FGeometryCollection::TransformToGeometryIndexAttribute, FGeometryCollection::TransformGroup);
+		const TManagedArray<FBox>& BoundingBoxes = GeomCollection->GetAttribute<FBox>(FGeometryCollection::BoundingBoxAttribute, FGeometryCollection::GeometryGroup);
 
 		TArray<FTransform> Transforms;
 		GeometryCollectionAlgo::GlobalMatrices(Transform, GeomCollection->Parent, Transforms);
@@ -1310,3 +1309,318 @@ int32 FFractureEngineFracturing::UniformFracture(
 	return INDEX_NONE;
 }
 
+namespace UE::Dataflow::Private
+{
+	static void SetBoneColor(FManagedArrayCollection& InCollection, const int32 InBoneIdx, const FLinearColor InBoneColor)
+	{
+		const TManagedArray<TSet<int32>>& Children = InCollection.GetAttribute<TSet<int32>>(FTransformCollection::ChildrenAttribute, FGeometryCollection::TransformGroup);
+		TManagedArray<FLinearColor>& BoneColors = InCollection.ModifyAttribute<FLinearColor>(FGeometryCollection::BoneColorAttribute, FGeometryCollection::TransformGroup);
+
+		if (Children[InBoneIdx].Num() == 0)
+		{
+			BoneColors[InBoneIdx] = InBoneColor;
+			return;
+		}
+
+		for (int32 Child : Children[InBoneIdx])
+		{
+			SetBoneColor(InCollection, Child, InBoneColor);
+		}
+	}
+
+	static FLinearColor GetRandomColor(const FRandomStream& InRandomStream)
+	{
+		const uint8 R = static_cast<uint8>(InRandomStream.FRandRange(5, 105));
+		const uint8 G = static_cast<uint8>(InRandomStream.FRandRange(5, 105));
+		const uint8 B = static_cast<uint8>(InRandomStream.FRandRange(5, 105));
+
+		return FLinearColor(FColor(R, G, B, 255));
+	}
+}
+
+void FFractureEngineFracturing::InitColors(FManagedArrayCollection& InCollection)
+{
+	//constexpr FLinearColor BoneColor = FLinearColor(0.01, 0.1, 0.1, 1.0);
+	const FLinearColor BoneColor = FLinearColor(FColor(255U, 255U, 255U, 255U));
+	constexpr FLinearColor VertexColor = FLinearColor(1.0, 1.0, 1.0, 1.0);
+
+	TManagedArray<FLinearColor>& BoneColors = InCollection.ModifyAttribute<FLinearColor>(FGeometryCollection::BoneColorAttribute, FGeometryCollection::TransformGroup);
+	const int32 NumBones = BoneColors.Num();
+	for (int32 BoneIdx = 0; BoneIdx < NumBones; ++BoneIdx)
+	{
+		BoneColors[BoneIdx] = BoneColor;
+	}
+
+	TManagedArray<FLinearColor>& VertexColors = InCollection.ModifyAttribute<FLinearColor>(FGeometryCollection::ColorAttribute, FGeometryCollection::VerticesGroup);
+	const int32 NumVertices = VertexColors.Num();
+	for (int32 VertexIdx = 0; VertexIdx < NumVertices; ++VertexIdx)
+	{
+		VertexColors[VertexIdx] = VertexColor;
+	}
+}
+
+void FFractureEngineFracturing::TransferBoneColorToVertexColor(FManagedArrayCollection& InCollection)
+{
+	const int32 NumTransforms = InCollection.NumElements(FGeometryCollection::TransformGroup);
+
+	TManagedArray<FLinearColor>& BoneColors = InCollection.ModifyAttribute<FLinearColor>(FGeometryCollection::BoneColorAttribute, FGeometryCollection::TransformGroup);
+	const TManagedArray<int32>& TransformToGeometryIndexArr = InCollection.GetAttribute<int32>(FGeometryCollection::TransformToGeometryIndexAttribute, FGeometryCollection::TransformGroup);
+	const TManagedArray<int32>& VertexStartArr = InCollection.GetAttribute<int32>(FGeometryCollection::VertexStartAttribute, FGeometryCollection::GeometryGroup);
+	const TManagedArray<int32>& VertexCountArr = InCollection.GetAttribute<int32>(FGeometryCollection::VertexCountAttribute, FGeometryCollection::GeometryGroup);
+
+	TManagedArray<FLinearColor>& VertexColorArr = InCollection.ModifyAttribute<FLinearColor>(FGeometryCollection::ColorAttribute, FGeometryCollection::VerticesGroup);
+
+	for (int32 TransformIdx = 0; TransformIdx < NumTransforms; ++TransformIdx)
+	{
+		const int32 GeometryIndex = TransformToGeometryIndexArr[TransformIdx];
+		// Only transfer color to non clusters
+		if (GeometryIndex != -1)
+		{
+			const int32 VertexStart = VertexStartArr[GeometryIndex];
+			const int32 VertexCount = VertexCountArr[GeometryIndex];
+
+			for (int32 VertexIdx = VertexStart; VertexIdx < VertexStart + VertexCount; ++VertexIdx)
+			{
+				VertexColorArr[VertexIdx] = BoneColors[TransformIdx];
+			}
+		}
+	}
+}
+
+void FFractureEngineFracturing::SetBoneColorByParent(
+	FManagedArrayCollection& InCollection,
+	const FRandomStream& InRandomStream,
+	int32 InLevel)
+{
+	TManagedArray<FLinearColor>& BoneColors = InCollection.ModifyAttribute<FLinearColor>(FGeometryCollection::BoneColorAttribute, FGeometryCollection::TransformGroup);
+	const int32 NumBones = BoneColors.Num();
+	const TManagedArray<int32>& Levels = InCollection.GetAttribute<int32>(FTransformCollection::LevelAttribute, FGeometryCollection::TransformGroup);
+	const TManagedArray<TSet<int32>>& Children = InCollection.GetAttribute<TSet<int32>>(FTransformCollection::ChildrenAttribute, FGeometryCollection::TransformGroup);
+
+	const FLinearColor BlankColor(FColor(255U, 255U, 255U, 255U));
+
+	if (InLevel == 0)
+	{
+		FLinearColor Color = UE::Dataflow::Private::GetRandomColor(InRandomStream);
+
+		for (int32 BoneIdx = 0; BoneIdx < NumBones; ++BoneIdx)
+		{
+			BoneColors[BoneIdx] = Color;
+		}
+	}
+	else if (InLevel > 0)
+	{
+		for (int32 BoneIdx = 0; BoneIdx < NumBones; ++BoneIdx)
+		{
+			if (Levels[BoneIdx] == InLevel - 1)
+			{
+				FLinearColor Color = UE::Dataflow::Private::GetRandomColor(InRandomStream);
+
+				for (int32 ChildBoneIdx : Children[BoneIdx])
+				{
+					UE::Dataflow::Private::SetBoneColor(InCollection, ChildBoneIdx, Color);
+				}
+			}
+		}
+	}
+}
+
+void FFractureEngineFracturing::SetBoneColorByLevel(FManagedArrayCollection& InCollection, int32 InLevel)
+{
+	const FColor ColorsPerDepth[] =
+	{
+		FColor::Cyan,
+		FColor::Orange,
+		FColor::Emerald,
+		FColor::Yellow,
+		FColor::Purple,
+		FColor::Magenta,
+		FColor::Turquoise,
+		FColor::Emerald,
+		FColor::Silver,
+		FColor::Blue
+	};
+
+	const FLinearColor BlankColor(FColor(255U, 255U, 255U, 255U));
+
+	TManagedArray<FLinearColor>& BoneColors = InCollection.ModifyAttribute<FLinearColor>(FGeometryCollection::BoneColorAttribute, FGeometryCollection::TransformGroup);
+	const int32 NumBones = BoneColors.Num();
+	const TManagedArray<int32>& Levels = InCollection.GetAttribute<int32>(FTransformCollection::LevelAttribute, FGeometryCollection::TransformGroup);
+
+	for (int32 BoneIdx = 0; BoneIdx < NumBones; ++BoneIdx)
+	{
+		if (InLevel >= 0)
+		{
+			if (Levels[BoneIdx] >= InLevel)
+			{
+				BoneColors[BoneIdx] = ColorsPerDepth[InLevel];
+			}
+			else
+			{
+				BoneColors[BoneIdx] = BlankColor;
+			}
+		}
+	}
+}
+
+ void FFractureEngineFracturing::SetBoneColorByCluster(
+	 FManagedArrayCollection& InCollection, 
+	 const FRandomStream& InRandomStream, 
+	 int32 InLevel)
+{
+	TArray<FLinearColor> RandomColors;
+
+	FRandomStream Random(1);
+	for (int i = 0; i < 100; i++)
+	{
+		const uint8 R = static_cast<uint8>(Random.FRandRange(5, 105));
+		const uint8 G = static_cast<uint8>(Random.FRandRange(5, 105));
+		const uint8 B = static_cast<uint8>(Random.FRandRange(5, 105));
+		RandomColors.Push(FLinearColor(FColor(R, G, B, 255)));
+	}
+
+	const FLinearColor BlankColor(FColor(255U, 255U, 255U, 255U));
+
+	const TManagedArray<int32>& Parents = InCollection.GetAttribute<int32>("Parent", FGeometryCollection::TransformGroup);
+
+	const TManagedArray<int32>& Levels = InCollection.GetAttribute<int32>(FTransformCollection::LevelAttribute, FGeometryCollection::TransformGroup);
+	TManagedArray<FLinearColor>& BoneColors = InCollection.ModifyAttribute<FLinearColor>(FGeometryCollection::BoneColorAttribute, FGeometryCollection::TransformGroup);
+
+	const int32 NumParents = Parents.Num();
+
+	for (int32 BoneIndex = 0, NumBones = NumParents; BoneIndex < NumBones; ++BoneIndex)
+	{
+		FLinearColor BoneColor = FLinearColor(FColor::Black);
+
+		if (Levels[BoneIndex] >= InLevel)
+		{
+			// go up until we find parent at the required ViewLevel
+			int32 Bone = BoneIndex;
+			while (Bone != -1 && Levels[Bone] > InLevel)
+			{
+				Bone = Parents[Bone];
+			}
+
+			int32 ColorIndex = Bone + 1; // parent can be -1 for root, range [-1..n]
+			BoneColor = RandomColors[ColorIndex % RandomColors.Num()];
+
+			BoneColor.LinearRGBToHSV();
+			BoneColor.B *= .5;
+			BoneColor.HSVToLinearRGB();
+		}
+		else
+		{
+			BoneColor = BlankColor;
+		}
+
+		BoneColors[BoneIndex] = BoneColor;
+	}
+}
+
+ void FFractureEngineFracturing::SetBoneColorByLeafLevel(FManagedArrayCollection& InCollection, int32 InLevel)
+ {
+	 const FColor ColorsPerDepth[] =
+	 {
+		 FColor::Cyan,
+		 FColor::Orange,
+		 FColor::Emerald,
+		 FColor::Yellow,
+		 FColor::Purple,
+		 FColor::Magenta,
+		 FColor::Turquoise,
+		 FColor::Emerald,
+		 FColor::Silver,
+		 FColor::Blue
+	 };
+
+	 const FLinearColor BlankColor(FColor(255U, 255U, 255U, 255U));
+
+	 TManagedArray<FLinearColor>& BoneColors = InCollection.ModifyAttribute<FLinearColor>(FGeometryCollection::BoneColorAttribute, FGeometryCollection::TransformGroup);
+	 const int32 NumBones = BoneColors.Num();
+	 const TManagedArray<int32>& Levels = InCollection.GetAttribute<int32>(FTransformCollection::LevelAttribute, FGeometryCollection::TransformGroup);
+	 const TManagedArray<TSet<int32>>& Children = InCollection.GetAttribute<TSet<int32>>(FTransformCollection::ChildrenAttribute, FGeometryCollection::TransformGroup);
+
+	 for (int32 BoneIdx = 0; BoneIdx < NumBones; ++BoneIdx)
+	 {
+		 if (Levels[BoneIdx] >= InLevel)
+		 {
+			 BoneColors[BoneIdx] = Children[BoneIdx].Num() == 0 ? ColorsPerDepth[Levels[BoneIdx]] : FColor::Black;
+		 }
+		 else
+		 {
+			 BoneColors[BoneIdx] = BlankColor;
+		 }
+	 }
+ }
+
+ void FFractureEngineFracturing::SetBoneColorByLeaf(FManagedArrayCollection& InCollection, const FRandomStream& InRandomStream, int32 InLevel)
+ {
+	 const FLinearColor BlankColor(FColor(255U, 255U, 255U, 255U));
+
+	 TManagedArray<FLinearColor>& BoneColors = InCollection.ModifyAttribute<FLinearColor>(FGeometryCollection::BoneColorAttribute, FGeometryCollection::TransformGroup);
+	 const TManagedArray<int32>& Levels = InCollection.GetAttribute<int32>(FTransformCollection::LevelAttribute, FGeometryCollection::TransformGroup);
+	 const int32 NumBones = BoneColors.Num();
+
+	 for (int32 BoneIdx = 0; BoneIdx < NumBones; ++BoneIdx)
+	 {
+		 if (Levels[BoneIdx] >= InLevel)
+		 {
+			 BoneColors[BoneIdx] = UE::Dataflow::Private::GetRandomColor(InRandomStream);
+		 }
+		 else
+		 {
+			 BoneColors[BoneIdx] = BlankColor;
+		 }
+	 }
+ }
+
+ void FFractureEngineFracturing::SetBoneColorByAttr(
+	 FManagedArrayCollection& InCollection,
+	 const FString InAttribute,
+	 const float InMinAttrValue,
+	 float InMaxAttrValue,
+	 const FLinearColor InMinColor,
+	 const FLinearColor InMaxColor)
+ {
+	 const FName AttrName = FName(*InAttribute);
+	 if (InCollection.HasAttribute(AttrName, FGeometryCollection::TransformGroup))
+	 {
+		 const TManagedArray<float>& AttrValues = InCollection.GetAttribute<float>(AttrName, FGeometryCollection::TransformGroup);
+		 TManagedArray<FLinearColor>& BoneColors = InCollection.ModifyAttribute<FLinearColor>(FGeometryCollection::BoneColorAttribute, FGeometryCollection::TransformGroup);
+
+		 for (int32 Idx = 0; Idx < AttrValues.Num(); ++Idx)
+		 {
+			 float AttrValue = AttrValues[Idx];
+
+			 // Make sure AttrValue is valid
+			 AttrValue = FMath::Max(AttrValue, InMinAttrValue);
+			 AttrValue = FMath::Min(AttrValue, InMaxAttrValue);
+			 if (InMaxAttrValue < InMinAttrValue)
+			 {
+				 InMaxAttrValue = InMinAttrValue + 0.01f;
+			 }
+
+			 // Lerp between InMinColor and InMaxColor
+			 float Alpha = (AttrValue - InMinAttrValue) / (InMaxAttrValue - InMinAttrValue);
+			 FLinearColor NewColor = FLinearColor::LerpUsingHSV(InMinColor, InMaxColor, Alpha);
+			 BoneColors[Idx] = NewColor;
+		 }
+	 }
+ }
+
+ void FFractureEngineFracturing::SetBoneColorRandom(
+	 FManagedArrayCollection& InCollection, 
+	 const FRandomStream& InRandomStream)
+ {
+	 TManagedArray<FLinearColor>& BoneColors = InCollection.ModifyAttribute<FLinearColor>(FGeometryCollection::BoneColorAttribute, FGeometryCollection::TransformGroup);
+	 const int32 NumBones = BoneColors.Num();
+
+	 for (int32 Idx = 0; Idx < NumBones; ++Idx)
+	 {
+		 const uint8 R = static_cast<uint8>(InRandomStream.FRandRange(5, 105));
+		 const uint8 G = static_cast<uint8>(InRandomStream.FRandRange(5, 105));
+		 const uint8 B = static_cast<uint8>(InRandomStream.FRandRange(5, 105));
+
+		 BoneColors[Idx] = FLinearColor(FColor(R, G, B, 255));
+	 }
+ }
diff --git a/Engine/Plugins/Experimental/Fracture/Source/FractureEngine/Public/FractureEngineFracturing.h b/Engine/Plugins/Experimental/Fracture/Source/FractureEngine/Public/FractureEngineFracturing.h
--- a/Engine/Plugins/Experimental/Fracture/Source/FractureEngine/Public/FractureEngineFracturing.h
+++ b/Engine/Plugins/Experimental/Fracture/Source/FractureEngine/Public/FractureEngineFracturing.h
@@ -68,8 +68,12 @@ struct FUniformFractureProcSettings
 class FRACTUREENGINE_API FFractureEngineFracturing
 {
 public:
-
-	static void GenerateExplodedViewAttribute(FManagedArrayCollection& InOutCollection, const FVector& InScale, float InUniformScale);
+	static void GenerateExplodedViewAttribute(
+		FManagedArrayCollection& InOutCollection,
+		const FVector& InScale,
+		const float InUniformScale,
+		const int32 InViewFractureLevel = -1,
+		const int32 InMaxFractureLevel = -1);
 
 	static int32 VoronoiFracture(FManagedArrayCollection& InOutCollection,
 		FDataflowTransformSelection InTransformSelection,
@@ -200,5 +204,43 @@ public:
 		FDataflowTransformSelection InTransformSelection,
 		const FUniformFractureSettings& InUniformFractureSettings);
 
+	static void InitColors(FManagedArrayCollection& InCollection);
+
+	static void TransferBoneColorToVertexColor(FManagedArrayCollection& InCollection);
+
+	static void SetBoneColorByParent(
+		FManagedArrayCollection& InCollection,
+		const FRandomStream& InRandomStream,
+		int32 InLevel);
+
+	static void SetBoneColorByLevel(
+		FManagedArrayCollection& InCollection, 
+		int32 InLevel);
+
+	static void SetBoneColorByCluster(
+		FManagedArrayCollection& InCollection,
+		const FRandomStream& InRandomStream,
+		int32 InLevel);
+
+	static void SetBoneColorByLeafLevel(
+		FManagedArrayCollection& InCollection, 
+		int32 InLevel);
+
+	static void SetBoneColorByLeaf(
+		FManagedArrayCollection& InCollection, 
+		const FRandomStream& InRandomStream,
+		int32 InLevel);
+
+	static void SetBoneColorByAttr(
+		FManagedArrayCollection& InCollection,
+		const FString InAttribute,
+		const float InMinAttrValue,
+		float InMaxAttrValue,
+		const FLinearColor InMinColor,
+		const FLinearColor InMaxColor);
+
+	static void SetBoneColorRandom(
+		FManagedArrayCollection& InCollection,
+		const FRandomStream& InRandomStream);
 };
 	
\ No newline at end of file
diff --git a/Engine/Plugins/Experimental/GeometryCollectionPlugin/Source/GeometryCollectionNodes/Private/Dataflow/GeometryCollectionClusteringNodes.cpp b/Engine/Plugins/Experimental/GeometryCollectionPlugin/Source/GeometryCollectionNodes/Private/Dataflow/GeometryCollectionClusteringNodes.cpp
--- a/Engine/Plugins/Experimental/GeometryCollectionPlugin/Source/GeometryCollectionNodes/Private/Dataflow/GeometryCollectionClusteringNodes.cpp
+++ b/Engine/Plugins/Experimental/GeometryCollectionPlugin/Source/GeometryCollectionNodes/Private/Dataflow/GeometryCollectionClusteringNodes.cpp
@@ -54,6 +54,40 @@ namespace UE::Dataflow
 	}
 }
 
+FAutoClusterDataflowNode::FAutoClusterDataflowNode(const UE::Dataflow::FNodeParameters& InParam, FGuid InGuid)
+	: FDataflowNode(InParam, InGuid)
+{
+	RegisterInputConnection(&Collection);
+	RegisterInputConnection(&TransformSelection);
+	RegisterInputConnection(&ClusterSites)
+		.SetCanHidePin(true)
+		.SetPinIsHidden(true);
+	RegisterInputConnection(&ClusterFraction)
+		.SetCanHidePin(true)
+		.SetPinIsHidden(true);
+	RegisterInputConnection(&SiteSize)
+		.SetCanHidePin(true)
+		.SetPinIsHidden(true);
+	RegisterInputConnection(&ClusterGridWidth)
+		.SetCanHidePin(true)
+		.SetPinIsHidden(true);
+	RegisterInputConnection(&ClusterGridDepth)
+		.SetCanHidePin(true)
+		.SetPinIsHidden(true);
+	RegisterInputConnection(&ClusterGridHeight)
+		.SetCanHidePin(true)
+		.SetPinIsHidden(true);
+	RegisterInputConnection(&MinimumSize)
+		.SetCanHidePin(true)
+		.SetPinIsHidden(true);
+	RegisterInputConnection(&bPreferConvexity)
+		.SetCanHidePin(true)
+		.SetPinIsHidden(true);
+	RegisterInputConnection(&ConcavityTolerance)
+		.SetCanHidePin(true)
+		.SetPinIsHidden(true);
+	RegisterOutputConnection(&Collection, &Collection);
+}
 
 void FAutoClusterDataflowNode::Evaluate(UE::Dataflow::FContext& Context, const FDataflowOutput* Out) const
 {
@@ -116,7 +150,7 @@ void FClusterFlattenDataflowNode::Evaluate(UE::Dataflow::FContext& Context, cons
 	if (Out->IsA(&Collection))
 	{
 		const FManagedArrayCollection& InCollection = GetValue(Context, &Collection);
-		if (InCollection.NumElements(FGeometryCollection::TransformAttribute) > 0)
+		if (InCollection.NumElements(FTransformCollection::TransformAttribute) > 0)
 		{
 			if (TUniquePtr<FGeometryCollection> GeomCollection = TUniquePtr<FGeometryCollection>(InCollection.NewCopy<FGeometryCollection>()))
 			{
@@ -164,7 +198,7 @@ void FClusterUnclusterDataflowNode::Evaluate(UE::Dataflow::FContext& Context, co
 	if (Out->IsA<FManagedArrayCollection>(&Collection))
 	{
 		const FManagedArrayCollection& InCollection = GetValue<FManagedArrayCollection>(Context, &Collection);
-		if (InCollection.NumElements(FGeometryCollection::TransformAttribute) > 0)
+		if (InCollection.NumElements(FTransformCollection::TransformAttribute) > 0)
 		{
 			const FDataflowTransformSelection& InTransformSelection = GetValue<FDataflowTransformSelection>(Context, &TransformSelection);
 			if (TUniquePtr<FGeometryCollection> GeomCollection = TUniquePtr<FGeometryCollection>(InCollection.NewCopy<FGeometryCollection>()))
diff --git a/Engine/Plugins/Experimental/GeometryCollectionPlugin/Source/GeometryCollectionNodes/Private/Dataflow/GeometryCollectionFracturingNodes.cpp b/Engine/Plugins/Experimental/GeometryCollectionPlugin/Source/GeometryCollectionNodes/Private/Dataflow/GeometryCollectionFracturingNodes.cpp
--- a/Engine/Plugins/Experimental/GeometryCollectionPlugin/Source/GeometryCollectionNodes/Private/Dataflow/GeometryCollectionFracturingNodes.cpp
+++ b/Engine/Plugins/Experimental/GeometryCollectionPlugin/Source/GeometryCollectionNodes/Private/Dataflow/GeometryCollectionFracturingNodes.cpp
@@ -60,6 +60,9 @@ namespace UE::Dataflow
 		DATAFLOW_NODE_REGISTER_CREATION_FACTORY(FBrickCutterDataflowNode);
 		DATAFLOW_NODE_REGISTER_CREATION_FACTORY(FMeshCutterDataflowNode);
 		DATAFLOW_NODE_REGISTER_CREATION_FACTORY(FUniformFractureDataflowNode);
+		DATAFLOW_NODE_REGISTER_CREATION_FACTORY(FVisualizeFractureDataflowNode);
+//		Commented out until we decide how to make generic data setter nodes
+//		DATAFLOW_NODE_REGISTER_CREATION_FACTORY(FSetFloatAttributeDataflowNode);
 
 		// GeometryCollection|Fracture
 		DATAFLOW_NODE_REGISTER_CREATION_FACTORY_NODE_COLORS_BY_CATEGORY("GeometryCollection|Fracture", FLinearColor(1.f, 1.f, 0.8f), CDefaultNodeBodyTintColor);
@@ -589,6 +592,17 @@ void FExplodedViewDataflowNode::Evaluate(UE::Dataflow::FContext& Context, const
 	{
 		FManagedArrayCollection InCollection = GetValue<FManagedArrayCollection>(Context, &Collection);
 
+		// Translate collection
+		FVector InOffset = GetValue(Context, &Offset);
+		if (InOffset.Length() > UE_KINDA_SMALL_NUMBER)
+		{
+			TManagedArray<FVector3f>& Vertex = InCollection.ModifyAttribute<FVector3f>("Vertex", FGeometryCollection::VerticesGroup);
+			for (int32 VertexIdx = 0; VertexIdx < Vertex.Num(); ++VertexIdx)
+			{
+				Vertex[VertexIdx] += FVector3f(InOffset);
+			}
+		}
+
 		FFractureEngineFracturing::GenerateExplodedViewAttribute(InCollection, GetValue<FVector>(Context, &Scale), GetValue<float>(Context, &UniformScale));
 
 		SetValue(Context, MoveTemp(InCollection), &Collection);
@@ -1098,4 +1112,224 @@ void FUniformFractureDataflowNode::Evaluate(UE::Dataflow::FContext& Context, con
 	}
 }
 
+/** ------------------------------------------------------------------------------------------------------------ **/
+
+FVisualizeFractureDataflowNode::FVisualizeFractureDataflowNode(const UE::Dataflow::FNodeParameters& InParam, FGuid InGuid)
+	: FDataflowNode(InParam, InGuid)
+{
+	RegisterInputConnection(&Collection);
+	RegisterInputConnection(&RandomSeed)
+		.SetCanHidePin(true)
+		.SetPinIsHidden(true);
+	RegisterInputConnection(&Level)
+		.SetCanHidePin(true)
+		.SetPinIsHidden(true);
+	RegisterInputConnection(&ExplodeAmount)
+		.SetCanHidePin(true)
+		.SetPinIsHidden(true);
+	RegisterInputConnection(&Scale)
+		.SetCanHidePin(true)
+		.SetPinIsHidden(true);
+	RegisterInputConnection(&Offset)
+		.SetCanHidePin(true)
+		.SetPinIsHidden(true);
+	RegisterOutputConnection(&Collection, &Collection);
+}
+
+void FVisualizeFractureDataflowNode::Evaluate(UE::Dataflow::FContext& Context, const FDataflowOutput* Out) const
+{
+	if (Out->IsA(&Collection))
+	{
+		FManagedArrayCollection InCollection = GetValue(Context, &Collection);
+
+		// Translate collection
+		FVector InOffset = GetValue(Context, &Offset);
+		if (InOffset.Length() > UE_KINDA_SMALL_NUMBER)
+		{
+			TManagedArray<FVector3f>& Vertex = InCollection.ModifyAttribute<FVector3f>("Vertex", FGeometryCollection::VerticesGroup);
+			for (int32 VertexIdx = 0; VertexIdx < Vertex.Num(); ++VertexIdx)
+			{
+				Vertex[VertexIdx] += FVector3f(InOffset);
+			}
+		}
+
+		const int32 InLevel = GetValue(Context, &Level);
+		check(InLevel >= 0);
+
+		if (bApplyExplodedView)
+		{
+			FFractureEngineFracturing::GenerateExplodedViewAttribute(
+				InCollection, 
+				GetValue(Context, &Scale), 
+				GetValue(Context, &ExplodeAmount),
+				InLevel);
+		}
+
+		if (bApplyColor)
+		{
+			if (InCollection.HasAttribute("BoneColor", FGeometryCollection::TransformGroup))
+			{
+				TManagedArray<FLinearColor>& BoneColors = InCollection.ModifyAttribute<FLinearColor>("BoneColor", FGeometryCollection::TransformGroup);
+
+				const int32 NumBones = BoneColors.Num();
+				const int32 InRandomSeed = GetValue(Context, &RandomSeed);
+
+				FRandomStream RandomStream(NumBones + InRandomSeed);
+
+				// Clear BoneColors
+				FFractureEngineFracturing::InitColors(InCollection);
+
+				if (ColoringType == EDataflowVisualizeFractureColoringType::ColorByParent && 
+					InCollection.HasAttribute(FTransformCollection::LevelAttribute, FGeometryCollection::TransformGroup))
+				{
+					FFractureEngineFracturing::SetBoneColorByParent(InCollection, RandomStream, InLevel);
+				}
+				else if (ColoringType == EDataflowVisualizeFractureColoringType::ColorByLevel &&
+					InCollection.HasAttribute(FTransformCollection::LevelAttribute, FGeometryCollection::TransformGroup))
+				{
+					FFractureEngineFracturing::SetBoneColorByLevel(InCollection, InLevel);
+				}
+				else if (ColoringType == EDataflowVisualizeFractureColoringType::ColorByCluster &&
+					InCollection.HasAttribute(FTransformCollection::LevelAttribute, FGeometryCollection::TransformGroup))
+				{
+					// This what the Geometry Tools uses
+					FFractureEngineFracturing::SetBoneColorByCluster(InCollection, RandomStream, InLevel);
+				}
+				else if (ColoringType == EDataflowVisualizeFractureColoringType::ColorByLeafLevel &&
+					InCollection.HasAttribute(FTransformCollection::LevelAttribute, FGeometryCollection::TransformGroup))
+				{
+					FFractureEngineFracturing::SetBoneColorByLeafLevel(InCollection, InLevel);
+				}
+				else if (ColoringType == EDataflowVisualizeFractureColoringType::ColorByLeaf)
+				{
+					FFractureEngineFracturing::SetBoneColorByLeaf(InCollection, RandomStream, InLevel);
+				}
+				else if (ColoringType == EDataflowVisualizeFractureColoringType::ColorByAttr)
+				{
+					FFractureEngineFracturing::SetBoneColorByAttr(InCollection,
+						Attribute,
+						Min.MinAttrValue,
+						Max.MaxAttrValue,
+						Min.MinColor,
+						Max.MaxColor);
+				}
+
+				// Transfer BoneColors to VertexColor
+				FFractureEngineFracturing::TransferBoneColorToVertexColor(InCollection);
+
+				SetValue(Context, MoveTemp(InCollection), &Collection);
+			}
+		}
+		else
+		{
+			SafeForwardInput(Context, &Collection, &Collection);
+		}
+	}
+}
 
+FSetFloatAttributeDataflowNode::FSetFloatAttributeDataflowNode(const UE::Dataflow::FNodeParameters& InParam, FGuid InGuid)
+	: FDataflowNode(InParam, InGuid)
+{
+	RegisterInputConnection(&Collection);
+	RegisterInputConnection(&RandomSeed)
+		.SetCanHidePin(true)
+		.SetPinIsHidden(true);
+	RegisterInputConnection(&NoiseScale)
+		.SetCanHidePin(true)
+		.SetPinIsHidden(true);
+	RegisterOutputConnection(&Collection, &Collection);
+	RegisterOutputConnection(&FloatArray);
+}
+
+void FSetFloatAttributeDataflowNode::Evaluate(UE::Dataflow::FContext& Context, const FDataflowOutput* Out) const
+{
+	if (Out->IsA(&Collection) || Out->IsA(&FloatArray))
+	{
+		if (IsConnected(&Collection))
+		{
+			FManagedArrayCollection InCollection = GetValue(Context, &Collection);
+			const FName AttrName = FName(*Attribute);
+
+			if (InCollection.HasAttribute(AttrName, FGeometryCollection::TransformGroup))
+			{
+				TManagedArray<float>& AttrValues = InCollection.ModifyAttribute<float>(AttrName, FGeometryCollection::TransformGroup);
+				const int32 NumAttrValues = AttrValues.Num();
+
+				TArray<float> OutFloatArray; OutFloatArray.AddUninitialized(NumAttrValues);
+
+				if (Method == EDataflowSetFloatArrayMethod::Random)
+				{
+					const int32 InRandomSeed = GetValue(Context, &RandomSeed);
+					FRandomStream RandomStream(InRandomSeed);
+
+					for (int32 Idx = 0; Idx < NumAttrValues; ++Idx)
+					{
+						AttrValues[Idx] = RandomStream.FRandRange(0.f, 1.f);
+						OutFloatArray[Idx] = AttrValues[Idx];
+					}
+				}
+				else if (Method == EDataflowSetFloatArrayMethod::Noise)
+				{
+					const int32 InNoiseScale = GetValue(Context, &NoiseScale);
+
+					const TManagedArray<int32>& TransformToGeometryIndices = InCollection.GetAttribute<int32>(FGeometryCollection::TransformToGeometryIndexAttribute, FGeometryCollection::TransformGroup);
+					const TManagedArray<FBox>& BoundingBoxes = InCollection.GetAttribute<FBox>(FGeometryCollection::BoundingBoxAttribute, FGeometryCollection::GeometryGroup);
+
+					for (int32 Idx = 0; Idx < NumAttrValues; ++Idx)
+					{
+						int32 GeometryIdx = TransformToGeometryIndices[Idx];
+						if (GeometryIdx != -1)
+						{
+							FVector Center = BoundingBoxes[GeometryIdx].GetCenter();
+
+							AttrValues[Idx] = 0.5f * FMath::PerlinNoise3D(InNoiseScale * Center) + 1.f;
+						}
+						else
+						{
+							AttrValues[Idx] = 0.f;
+						}
+					}
+				}
+				else if (Method == EDataflowSetFloatArrayMethod::ByBoundingBox)
+				{
+					const TManagedArray<int32>& TransformToGeometryIndices = InCollection.GetAttribute<int32>(FGeometryCollection::TransformToGeometryIndexAttribute, FGeometryCollection::TransformGroup);
+					const TManagedArray<FBox>& BoundingBoxes = InCollection.GetAttribute<FBox>(FGeometryCollection::BoundingBoxAttribute, FGeometryCollection::GeometryGroup);
+
+					// Compute BoundingBox for the entire collection
+					FBox BBox = FBox(ForceInit);
+
+					for (int32 Idx = 0; Idx < NumAttrValues; ++Idx)
+					{
+						int32 GeometryIdx = TransformToGeometryIndices[Idx];
+						if (GeometryIdx != -1)
+						{
+							BBox += BoundingBoxes[GeometryIdx];
+						}
+					}
+
+					for (int32 Idx = 0; Idx < NumAttrValues; ++Idx)
+					{
+						int32 GeometryIdx = TransformToGeometryIndices[Idx];
+						if (GeometryIdx != -1)
+						{
+							FVector Center = BoundingBoxes[GeometryIdx].GetCenter();
+
+							AttrValues[Idx] = (Center.X - BBox.Min.X) / (BBox.Max.X - BBox.Min.X);
+						}
+						else
+						{
+							AttrValues[Idx] = 0.f;
+						}
+					}
+				}
+
+				SetValue(Context, MoveTemp(InCollection), &Collection);
+				SetValue(Context, OutFloatArray, &FloatArray);
+				return;
+			}
+		}
+
+		SafeForwardInput(Context, &Collection, &Collection);
+		SetValue(Context, TArray<float>(), &FloatArray);
+	}
+}
diff --git a/Engine/Plugins/Experimental/GeometryCollectionPlugin/Source/GeometryCollectionNodes/Public/Dataflow/GeometryCollectionClusteringNodes.h b/Engine/Plugins/Experimental/GeometryCollectionPlugin/Source/GeometryCollectionNodes/Public/Dataflow/GeometryCollectionClusteringNodes.h
--- a/Engine/Plugins/Experimental/GeometryCollectionPlugin/Source/GeometryCollectionNodes/Public/Dataflow/GeometryCollectionClusteringNodes.h
+++ b/Engine/Plugins/Experimental/GeometryCollectionPlugin/Source/GeometryCollectionNodes/Public/Dataflow/GeometryCollectionClusteringNodes.h
@@ -106,26 +106,9 @@ public:
 	UPROPERTY(meta = (DataflowInput, DisplayName = "TransformSelection"))
 	FDataflowTransformSelection TransformSelection;
 
-	FAutoClusterDataflowNode(const UE::Dataflow::FNodeParameters& InParam, FGuid InGuid = FGuid::NewGuid())
-		: FDataflowNode(InParam, InGuid)
-	{
-		RegisterInputConnection(&Collection);
-		RegisterInputConnection(&TransformSelection);
-		RegisterInputConnection(&ClusterSites);
-		RegisterInputConnection(&ClusterFraction);
-		RegisterInputConnection(&SiteSize);
-		RegisterInputConnection(&ClusterGridWidth);
-		RegisterInputConnection(&ClusterGridDepth);
-		RegisterInputConnection(&ClusterGridHeight);
-		RegisterInputConnection(&MinimumSize);
-		RegisterInputConnection(&bPreferConvexity);
-		RegisterInputConnection(&ConcavityTolerance);
-
-		RegisterOutputConnection(&Collection, &Collection);
-	}
+	FAutoClusterDataflowNode(const UE::Dataflow::FNodeParameters& InParam, FGuid InGuid = FGuid::NewGuid());
 
 	virtual void Evaluate(UE::Dataflow::FContext& Context, const FDataflowOutput* Out) const override;
-
 };
 
 
diff --git a/Engine/Plugins/Experimental/GeometryCollectionPlugin/Source/GeometryCollectionNodes/Public/Dataflow/GeometryCollectionFracturingNodes.h b/Engine/Plugins/Experimental/GeometryCollectionPlugin/Source/GeometryCollectionNodes/Public/Dataflow/GeometryCollectionFracturingNodes.h
--- a/Engine/Plugins/Experimental/GeometryCollectionPlugin/Source/GeometryCollectionNodes/Public/Dataflow/GeometryCollectionFracturingNodes.h
+++ b/Engine/Plugins/Experimental/GeometryCollectionPlugin/Source/GeometryCollectionNodes/Public/Dataflow/GeometryCollectionFracturingNodes.h
@@ -60,6 +60,7 @@ struct FUniformScatterPointsDataflowNode_v2 : public FDataflowNode
 {
 	GENERATED_USTRUCT_BODY()
 	DATAFLOW_NODE_DEFINE_INTERNAL(FUniformScatterPointsDataflowNode_v2, "UniformScatterPoints", "Generators|Point", "")
+	DATAFLOW_NODE_RENDER_TYPE("PointsRender", FName("TArray<FVector>"), "Points")
 
 public:
 	/** Minimum for the random range */
@@ -153,6 +154,7 @@ struct FRadialScatterPointsDataflowNode_v2 : public FDataflowNode
 {
 	GENERATED_USTRUCT_BODY()
 	DATAFLOW_NODE_DEFINE_INTERNAL(FRadialScatterPointsDataflowNode_v2, "RadialScatterPoints", "Generators|Point", "")
+	DATAFLOW_NODE_RENDER_TYPE("PointsRender", FName("TArray<FVector>"), "Points")
 
 public:
 	/** BoundingBox to generate points inside of */
@@ -248,6 +250,7 @@ struct FGridScatterPointsDataflowNode : public FDataflowNode
 {
 	GENERATED_USTRUCT_BODY()
 	DATAFLOW_NODE_DEFINE_INTERNAL(FGridScatterPointsDataflowNode, "GridScatterPoints", "Generators|Point", "")
+	DATAFLOW_NODE_RENDER_TYPE("PointsRender", FName("TArray<FVector>"), "Points")
 
 public:
 	/** Number of points in X direction */
@@ -728,12 +731,23 @@ public:
 	UPROPERTY(EditAnywhere, Category = "Scale", meta = (DataflowInput))
 	FVector Scale = FVector(1.0);
 
+	/** Translate collection for exploded view */
+	UPROPERTY(EditAnywhere, Category = "Exploded View", meta = (DataflowInput))
+	FVector Offset = FVector(0.0);
+
 	FExplodedViewDataflowNode(const UE::Dataflow::FNodeParameters& InParam, FGuid InGuid = FGuid::NewGuid())
 		: FDataflowNode(InParam, InGuid)
 	{
 		RegisterInputConnection(&Collection);
-		RegisterInputConnection(&UniformScale);
-		RegisterInputConnection(&Scale);
+		RegisterInputConnection(&UniformScale)
+			.SetCanHidePin(true)
+			.SetPinIsHidden(true);
+		RegisterInputConnection(&Scale)
+			.SetCanHidePin(true)
+			.SetPinIsHidden(true);
+		RegisterInputConnection(&Offset)
+			.SetCanHidePin(true)
+			.SetPinIsHidden(true);
 		RegisterOutputConnection(&Collection, &Collection);
 	}
 
@@ -1278,6 +1292,159 @@ public:
 	virtual void Evaluate(UE::Dataflow::FContext& Context, const FDataflowOutput* Out) const override;
 };
 
+UENUM(BlueprintType)
+enum class EDataflowVisualizeFractureColoringType : uint8
+{
+	ColorByParent UMETA(DisplayName = "Color by Parent"),
+	ColorByLevel UMETA(DisplayName = "Color by Level"),
+	ColorByCluster UMETA(DisplayName = "Color by Cluster"),
+	ColorByLeafLevel UMETA(DisplayName = "Color by Leaf Level"),
+	ColorByLeaf UMETA(DisplayName = "Color by Leaf"),
+	ColorByAttr UMETA(DisplayName = "Color by Attribute"),
+};
+
+USTRUCT() 
+struct FMinSettings
+{
+	GENERATED_USTRUCT_BODY()
+
+	UPROPERTY(EditAnywhere, Category = "Color", meta = (DisplayName = "Minimum Attr Value", EditCondition = "ColoringType == EDataflowVisualizeFractureColoringType::ColorByAttr", EditConditionHides))
+	float MinAttrValue = 0.f;
+
+	UPROPERTY(EditAnywhere, Category = "Color", meta = (EditCondition = "ColoringType == EDataflowVisualizeFractureColoringType::ColorByAttr", EditConditionHides))
+	FLinearColor MinColor = FLinearColor(FLinearColor::Green);
+};
+
+USTRUCT()
+struct FMaxSettings
+{
+	GENERATED_USTRUCT_BODY()
+
+	UPROPERTY(EditAnywhere, Category = "Color", meta = (DisplayName = "Maximum Attr Value", EditCondition = "ColoringType == EDataflowVisualizeFractureColoringType::ColorByAttr", EditConditionHides))
+	float MaxAttrValue = 1.f;
+
+	UPROPERTY(EditAnywhere, Category = "Color", meta = (EditCondition = "ColoringType == EDataflowVisualizeFractureColoringType::ColorByAttr", EditConditionHides))
+	FLinearColor MaxColor = FLinearColor(FLinearColor::Red);
+};
+
+/**
+ * 
+ * Visualizing fracture/cluster info in fractured collection
+ * 
+ */
+USTRUCT(meta = (DataflowGeometryCollection))
+struct FVisualizeFractureDataflowNode : public FDataflowNode
+{
+	GENERATED_USTRUCT_BODY()
+	DATAFLOW_NODE_DEFINE_INTERNAL(FVisualizeFractureDataflowNode, "VisualizeFracture", "GeometryCollection|Fracture|Utilities", "")
+	DATAFLOW_NODE_RENDER_TYPE("SurfaceRender", FGeometryCollection::StaticType(), "Collection")
+
+private:
+	/** Collection to visualize */
+	UPROPERTY(meta = (DataflowInput, DataflowOutput, DataflowPassthrough = "Collection", DataflowIntrinsic))
+	FManagedArrayCollection Collection;
+
+	/**  */
+	UPROPERTY(EditAnywhere, Category = "Level", meta = (DataflowInput, UIMin = "0", ClampMin = "0"));
+	int32 Level = 0;
+
+	/** Seed for random */
+	UPROPERTY(EditAnywhere, Category = "Random", meta = (DataflowInput, UIMin = "0", ClampMin = "0"));
+	int32 RandomSeed = 0;
+
+	/** Use cluster level for coloring and explode */
+	UPROPERTY(EditAnywhere, Category = "Exploded View");
+	bool bApplyExplodedView = true;
+
+	/** Scale amount to expand the pieces uniformly in all directions */
+	UPROPERTY(EditAnywhere, Category = "Exploded View", meta = (DataflowInput, UIMin = "0.0", ClampMin = "0.0", UIMax = "1.0", ClampMax = "1.0", EditCondition = "bApplyExplodedView"))
+	float ExplodeAmount = 0.5f;
+
+	/** Scale amounts to expand the pieces in all 3 directions */
+	UPROPERTY(EditAnywhere, Category = "Exploded View", meta = (DataflowInput, EditCondition = "bApplyExplodedView"))
+	FVector Scale = FVector(1.0);
+
+	/**  */
+	UPROPERTY(EditAnywhere, Category = "Color");
+	bool bApplyColor = true;
+
+	/**  */
+	UPROPERTY(EditAnywhere, Category = "Color", meta = (EditCondition = "bApplyColor"));
+	EDataflowVisualizeFractureColoringType ColoringType = EDataflowVisualizeFractureColoringType::ColorByLevel;
+
+	/**  */
+	UPROPERTY(EditAnywhere, Category = "Color", meta = (EditCondition = "ColoringType == EDataflowVisualizeFractureColoringType::ColorByAttr", EditConditionHides))
+	FString Attribute;
+
+	/**  */
+	UPROPERTY(EditAnywhere, Category = "Color", meta = (EditCondition = "ColoringType == EDataflowVisualizeFractureColoringType::ColorByAttr", EditConditionHides))
+	FMinSettings Min;
+
+	/**  */
+	UPROPERTY(EditAnywhere, Category = "Color", meta = (EditCondition = "ColoringType == EDataflowVisualizeFractureColoringType::ColorByAttr", EditConditionHides))
+	FMaxSettings Max;
+
+	/** Translate collection for exploded view */
+	UPROPERTY(EditAnywhere, Category = "Offset", meta = (DataflowInput))
+	FVector Offset = FVector(0.0);
+
+public:
+	FVisualizeFractureDataflowNode(const UE::Dataflow::FNodeParameters& InParam, FGuid InGuid = FGuid::NewGuid());
+
+	virtual void Evaluate(UE::Dataflow::FContext& Context, const FDataflowOutput* Out) const override;
+};
+
+UENUM(BlueprintType)
+enum class EDataflowSetFloatArrayMethod : uint8
+{
+	Random UMETA(DisplayName = "Random"),
+	Noise UMETA(DisplayName = "Noise"),
+	ByBoundingBox UMETA(DisplayName = "By BoundinBox"),
+};
+
+/**
+ *
+ * Set a float values in an array
+ *
+ */
+USTRUCT(meta = (DataflowGeometryCollection))
+struct FSetFloatAttributeDataflowNode : public FDataflowNode
+{
+	GENERATED_USTRUCT_BODY()
+	DATAFLOW_NODE_DEFINE_INTERNAL(FSetFloatAttributeDataflowNode, "SetFloatAttribute", "GeometryCollection|Utilities", "")
+	DATAFLOW_NODE_RENDER_TYPE("SurfaceRender", FGeometryCollection::StaticType(), "Collection")
+
+private:
+	/** Collection to visualize */
+	UPROPERTY(meta = (DataflowInput, DataflowOutput, DataflowPassthrough = "Collection", DataflowIntrinsic))
+	FManagedArrayCollection Collection;
+
+	/**  */
+	UPROPERTY(EditAnywhere, Category = "Attribute");
+	FString Attribute;
+
+	/**  */
+	UPROPERTY(EditAnywhere, Category = "Method");
+	EDataflowSetFloatArrayMethod Method = EDataflowSetFloatArrayMethod::Random;
+
+	/** Seed for random */
+	UPROPERTY(EditAnywhere, Category = "Random", meta = (DataflowInput, UIMin = "0", EditCondition = "Method == EDataflowSetFloatArrayMethod::Random", EditConditionHides));
+	int32 RandomSeed = 0;
+
+	/** Seed for random */
+	UPROPERTY(EditAnywhere, Category = "Random", meta = (DataflowInput, UIMin = "0", EditCondition = "Method == EDataflowSetFloatArrayMethod::Noise", EditConditionHides));
+	float NoiseScale = 1.f;
+
+	/** Output float array */
+	UPROPERTY(meta = (DataflowOutput))
+	TArray<float> FloatArray;
+
+public:
+	FSetFloatAttributeDataflowNode(const UE::Dataflow::FNodeParameters& InParam, FGuid InGuid = FGuid::NewGuid());
+
+	virtual void Evaluate(UE::Dataflow::FContext& Context, const FDataflowOutput* Out) const override;
+};
+
 namespace UE::Dataflow
 {
 	void GeometryCollectionFracturingNodes();
diff --git a/Engine/Source/Runtime/Experimental/Chaos/Private/GeometryCollection/GeometryCollection.cpp b/Engine/Source/Runtime/Experimental/Chaos/Private/GeometryCollection/GeometryCollection.cpp
--- a/Engine/Source/Runtime/Experimental/Chaos/Private/GeometryCollection/GeometryCollection.cpp
+++ b/Engine/Source/Runtime/Experimental/Chaos/Private/GeometryCollection/GeometryCollection.cpp
@@ -42,6 +42,12 @@ const FName FGeometryCollection::SimulationTypeAttribute("SimulationType");
 const FName FGeometryCollection::StatusFlagsAttribute("StatusFlags");
 const FName FGeometryCollection::ExternalCollisionsAttribute("ExternalCollisions");
 const FName FGeometryCollection::ColorAttribute("Color");
+const FName FGeometryCollection::BoneColorAttribute("BoneColor");
+const FName FGeometryCollection::TransformToGeometryIndexAttribute("TransformToGeometryIndex");
+const FName FGeometryCollection::BoundingBoxAttribute("BoundingBox");
+const FName FGeometryCollection::TransformIndexAttribute("TransformIndex");
+const FName FGeometryCollection::VertexStartAttribute("VertexStart");
+const FName FGeometryCollection::VertexCountAttribute("VertexCount");
 
 bool FGeometryCollection::AreCollisionParticlesEnabled()
 {
diff --git a/Engine/Source/Runtime/Experimental/Chaos/Public/GeometryCollection/GeometryCollection.h b/Engine/Source/Runtime/Experimental/Chaos/Public/GeometryCollection/GeometryCollection.h
--- a/Engine/Source/Runtime/Experimental/Chaos/Public/GeometryCollection/GeometryCollection.h
+++ b/Engine/Source/Runtime/Experimental/Chaos/Public/GeometryCollection/GeometryCollection.h
@@ -107,6 +107,12 @@ public:
 	static CHAOS_API const FName StatusFlagsAttribute;
 	static CHAOS_API const FName ExternalCollisionsAttribute;
 	static CHAOS_API const FName ColorAttribute;
+	static CHAOS_API const FName BoneColorAttribute;
+	static CHAOS_API const FName TransformToGeometryIndexAttribute;
+	static CHAOS_API const FName BoundingBoxAttribute;
+	static CHAOS_API const FName TransformIndexAttribute;
+	static CHAOS_API const FName VertexStartAttribute;
+	static CHAOS_API const FName VertexCountAttribute;
 
 	enum ESimulationTypes : uint8
 	{
