make dataflow fracture nodes convert selection to rigid/leaf nodes before fracturing

diff --git a/Engine/Plugins/Experimental/Fracture/Source/FractureEngine/Private/FractureEngineFracturing.cpp b/Engine/Plugins/Experimental/Fracture/Source/FractureEngine/Private/FractureEngineFracturing.cpp
--- a/Engine/Plugins/Experimental/Fracture/Source/FractureEngine/Private/FractureEngineFracturing.cpp
+++ b/Engine/Plugins/Experimental/Fracture/Source/FractureEngine/Private/FractureEngineFracturing.cpp
@@ -179,19 +179,46 @@ static void RandomReduceSelection(FDataflowTransformSelection& InOutTransformSel
 {
 	FRandomStream RandStream(InRandomSeed);
 
-	TArray<int32> SelectedBones = InOutTransformSelection.AsArray();
-
-	for (int32 i = 0; i < SelectedBones.Num(); i++)
+	for (int32 BoneIdx = 0; BoneIdx < InOutTransformSelection.Num(); ++BoneIdx)
 	{
-		if (RandStream.GetFraction() >= InProbToKeep) // range does not include 1, so if ProbToKeep is 1 this will never remove
+		if (InOutTransformSelection.IsSelected(BoneIdx))
 		{
-			SelectedBones.RemoveAtSwap(i);
-			i--;
+			if (RandStream.GetFraction() >= InProbToKeep) // range does not include 1, so if ProbToKeep is 1 this will never remove
+			{
+				InOutTransformSelection.SetNotSelected(BoneIdx);
+			}
 		}
 	}
+}
 
-	InOutTransformSelection.Clear();
-	InOutTransformSelection.SetFromArray(SelectedBones);
+namespace UE::Private::FractureHelpers
+{
+	static void SelectLeavesHelper(const FGeometryCollection& GeometryCollection, FDataflowTransformSelection& InOutTransformSelection, int32 BoneIdx)
+	{
+		if (GeometryCollection.SimulationType[BoneIdx] != FGeometryCollection::ESimulationTypes::FST_Rigid)
+		{
+			InOutTransformSelection.SetNotSelected(BoneIdx);
+			for (int32 ChildIdx : GeometryCollection.Children[BoneIdx])
+			{
+				SelectLeavesHelper(GeometryCollection, InOutTransformSelection, ChildIdx);
+			}
+		}
+		else
+		{
+			InOutTransformSelection.SetSelected(BoneIdx);
+		}
+	}
+	
+	static void ConvertToLeafSelection(const FGeometryCollection& GeometryCollection, FDataflowTransformSelection& InOutTransformSelection)
+	{
+		for (int32 BoneIdx = 0; BoneIdx < InOutTransformSelection.Num(); ++BoneIdx)
+		{
+			if (InOutTransformSelection.IsSelected(BoneIdx))
+			{
+				SelectLeavesHelper(GeometryCollection, InOutTransformSelection, BoneIdx);
+			}
+		}
+	}
 }
 
 
@@ -274,6 +301,7 @@ int32 FFractureEngineFracturing::VoronoiFracture(FManagedArrayCollection& InOutC
 				VoronoiPlanarCells.InternalSurfaceMaterials.NoiseSettings = NoiseSettings;
 
 				RandomReduceSelection(InTransformSelection, InRandomSeed, InChanceToFracture);
+				UE::Private::FractureHelpers::ConvertToLeafSelection(*GeomCollection, InTransformSelection);
 				TArray<int32> TransformSelectionArr = InTransformSelection.AsArray();
 
 				if (!FFractureEngineSelection::IsBoneSelectionValid(InOutCollection, TransformSelectionArr))
@@ -354,6 +382,7 @@ int32 FFractureEngineFracturing::PlaneCutter(FManagedArrayCollection& InOutColle
 		float GroutVal = InGrout;
 
 		RandomReduceSelection(InTransformSelection, InRandomSeed, InChanceToFracture);
+		UE::Private::FractureHelpers::ConvertToLeafSelection(*GeomCollection, InTransformSelection);
 		TArray<int32> TransformSelectionArr = InTransformSelection.AsArray();
 
 		if (!FFractureEngineSelection::IsBoneSelectionValid(InOutCollection, TransformSelectionArr))
@@ -503,6 +532,7 @@ int32 FFractureEngineFracturing::SliceCutter(FManagedArrayCollection& InOutColle
 		float GroutVal = InGrout;
 
 		RandomReduceSelection(InTransformSelection, InRandomSeed, InChanceToFracture);
+		UE::Private::FractureHelpers::ConvertToLeafSelection(*GeomCollection, InTransformSelection);
 		TArray<int32> TransformSelectionArr = InTransformSelection.AsArray();
 
 		if (!FFractureEngineSelection::IsBoneSelectionValid(InOutCollection, TransformSelectionArr))
@@ -877,6 +907,7 @@ int32 FFractureEngineFracturing::BrickCutter(FManagedArrayCollection& InOutColle
 		float CollisionSampleSpacingVal = InCollisionSampleSpacing;
 
 		RandomReduceSelection(InTransformSelection, InRandomSeed, InChanceToFracture);
+		UE::Private::FractureHelpers::ConvertToLeafSelection(*GeomCollection, InTransformSelection);
 		TArray<int32> TransformSelectionArr = InTransformSelection.AsArray();
 
 		if (!FFractureEngineSelection::IsBoneSelectionValid(InOutCollection, TransformSelectionArr))
@@ -999,6 +1030,7 @@ int32 FFractureEngineFracturing::MeshCutter(TArray<FTransform>& MeshTransforms,
 		const int32 OriginalNumTransforms = InOutCollection.NumElements(FGeometryCollection::TransformGroup);
 
 		RandomReduceSelection(InTransformSelection, InRandomSeed, InChanceToFracture);
+		UE::Private::FractureHelpers::ConvertToLeafSelection(*GeomCollection, InTransformSelection);
 		TArray<int32> TransformSelectionArr = InTransformSelection.AsArray();
 
 		for (const FTransform& ScatterTransform : MeshTransforms)
@@ -1201,6 +1233,7 @@ int32 FFractureEngineFracturing::UniformFracture(
 	if (TUniquePtr<FGeometryCollection> GeomCollection = TUniquePtr<FGeometryCollection>(InOutCollection.NewCopy<FGeometryCollection>()))
 	{
 		RandomReduceSelection(InTransformSelection, InUniformFractureSettings.RandomSeed, InUniformFractureSettings.ChanceToFracture);
+		UE::Private::FractureHelpers::ConvertToLeafSelection(*GeomCollection, InTransformSelection);
 		TArray<int32> TransformSelectionArr = InTransformSelection.AsArray();
 
 		if (!FFractureEngineSelection::IsBoneSelectionValid(InOutCollection, TransformSelectionArr))
