Dataflow:
- Exposed TransformLevelColors in Dataflow plugin settings, Fracture Tools and FVisualizeFractureDataflowNode using the plugin color settings
- Added DATAFLOW_NODE_REGISTER_CREATION_FACTORY_PIN_COLORS_BY_TYPE() macro to set PinTypeColors similarly to set NodeTypeColors
- Added Point rendering to Sampling nodes
- Bug fixes

diff --git a/Engine/Plugins/Experimental/ChaosEditor/Source/FractureEditor/FractureEditor.Build.cs b/Engine/Plugins/Experimental/ChaosEditor/Source/FractureEditor/FractureEditor.Build.cs
--- a/Engine/Plugins/Experimental/ChaosEditor/Source/FractureEditor/FractureEditor.Build.cs
+++ b/Engine/Plugins/Experimental/ChaosEditor/Source/FractureEditor/FractureEditor.Build.cs
@@ -1,6 +1,7 @@
 // Copyright Epic Games, Inc. All Rights Reserved.
 
 using UnrealBuildTool;
+using UnrealBuildTool.Rules;
 
 public class FractureEditor : ModuleRules
 {
@@ -48,7 +49,8 @@ public class FractureEditor : ModuleRules
 				"DeveloperSettings",
 				"FractureEngine",
 				"ContentBrowserData",
-				"GeometryFramework"
+				"GeometryFramework",
+				"DataflowCore"
 
 				// ... add private dependencies that you statically link with here ...	
 			}
diff --git a/Engine/Plugins/Experimental/ChaosEditor/Source/FractureEditor/Private/SGeometryCollectionOutliner.cpp b/Engine/Plugins/Experimental/ChaosEditor/Source/FractureEditor/Private/SGeometryCollectionOutliner.cpp
--- a/Engine/Plugins/Experimental/ChaosEditor/Source/FractureEditor/Private/SGeometryCollectionOutliner.cpp
+++ b/Engine/Plugins/Experimental/ChaosEditor/Source/FractureEditor/Private/SGeometryCollectionOutliner.cpp
@@ -23,6 +23,7 @@
 #include "GeometryCollection/Facades/CollectionHierarchyFacade.h"
 #include "GeometryCollection/GeometryCollectionSimulationTypes.h"
 #include "PhysicsProxy/GeometryCollectionPhysicsProxy.h"
+#include "Dataflow/DataflowSettings.h"
 
 #include UE_INLINE_GENERATED_CPP_BY_NAME(SGeometryCollectionOutliner)
 
@@ -104,15 +105,10 @@ void FGeometryCollectionTreeItem::GenerateContextMenu(UToolMenu* Menu, SGeometry
 
 FColor FGeometryCollectionTreeItem::GetColorPerDepth(uint32 Depth)
 {
-	const FColor ColorsPerDepth[] =
-	{
-		FColor::Cyan,
-		FColor::Orange,
-		FColor::Emerald,
-		FColor::Yellow,
-	};
-	const uint32 ColorPerDepthCount = sizeof(ColorsPerDepth) / sizeof(FColor);
-	return ColorsPerDepth[Depth % ColorPerDepthCount];
+	const UDataflowSettings* DataflowSettings = GetDefault<UDataflowSettings>();
+
+	const uint32 ColorPerDepthCount = DataflowSettings->TransformLevelColors.LevelColors.Num();
+	return DataflowSettings->TransformLevelColors.LevelColors[Depth % ColorPerDepthCount].ToFColor(true);
 }
 
 void FGeometryCollectionTreeItemBone::OnDragEnter(FDragDropEvent const& InDragDropEvent)
@@ -189,7 +185,7 @@ FReply FGeometryCollectionTreeItemBone::OnDrop(const FDragDropEvent& DragDropEve
 
 UOutlinerSettings::UOutlinerSettings(const FObjectInitializer& ObjInit)
 	: Super(ObjInit)
-	, ColorByLevel(false)
+	, ColorByLevel(true)
 	, ColumnMode(EOutlinerColumnMode::State)
 {}
 
diff --git a/Engine/Plugins/Experimental/Dataflow/Source/DataflowEditor/Private/Dataflow/DataflowSchema.cpp b/Engine/Plugins/Experimental/Dataflow/Source/DataflowEditor/Private/Dataflow/DataflowSchema.cpp
--- a/Engine/Plugins/Experimental/Dataflow/Source/DataflowEditor/Private/Dataflow/DataflowSchema.cpp
+++ b/Engine/Plugins/Experimental/Dataflow/Source/DataflowEditor/Private/Dataflow/DataflowSchema.cpp
@@ -22,25 +22,6 @@
 
 #define LOCTEXT_NAMESPACE "DataflowNode"
 
-namespace UE::Dataflow::Private
-{
-	static const FName ManagedArrayCollectionType = FName("FManagedArrayCollection");
-	static const FName FloatType = FName("float");
-	static const FName DoubleType = FName("double");
-	static const FName Int32Type = FName("int32");
-	static const FName BoolType = FName("bool");
-	static const FName StringType = FName("FString");
-	static const FName NameType = FName("FName");
-	static const FName TextType = FName("FText");
-	static const FName VectorType = FName("FVector");
-	static const FName TransformType = FName("FTransform");
-	static const FName RotatorType = FName("FRotator");
-	static const FName ArrayType = FName("TArray");
-	static const FName BoxType = FName("FBox");
-	static const FName SphereType = FName("FSphere");
-	static const FName DataflowAnyTypeType = FName("FDataflowAnyType");
-} // namespace UE::Dataflow::Private
-
 UDataflowSchema::UDataflowSchema()
 {
 }
@@ -305,65 +286,9 @@ FLinearColor UDataflowSchema::GetTypeColor(const FName& Type)
 	const UGraphEditorSettings* Settings = GetDefault<UGraphEditorSettings>();
 	const UDataflowSettings* DataflowSettings = GetDefault<UDataflowSettings>();
 
-	if (Type == UE::Dataflow::Private::ManagedArrayCollectionType)
-	{
-		return DataflowSettings->ManagedArrayCollectionPinTypeColor;
-	}
-	else if (Type == UE::Dataflow::Private::FloatType)
-	{
-		return Settings->FloatPinTypeColor;
-	}
-	else if (Type == UE::Dataflow::Private::DoubleType)
-	{
-		return Settings->DoublePinTypeColor;
-	}
-	else if (Type == UE::Dataflow::Private::Int32Type)
-	{
-		return Settings->IntPinTypeColor;
-	}
-	else if (Type == UE::Dataflow::Private::BoolType)
-	{
-		return Settings->BooleanPinTypeColor;
-	}
-	else if (Type == UE::Dataflow::Private::StringType)
-	{
-		return Settings->StringPinTypeColor;
-	}
-	else if (Type == UE::Dataflow::Private::NameType)
-	{
-		return Settings->NamePinTypeColor;
-	}
-	else if (Type == UE::Dataflow::Private::TextType)
-	{
-		return Settings->TextPinTypeColor;
-	}
-	else if (Type == UE::Dataflow::Private::VectorType)
-	{
-		return Settings->VectorPinTypeColor;
-	}
-	else if (Type == UE::Dataflow::Private::TransformType)
-	{
-		return Settings->TransformPinTypeColor;
-	}
-	else if (Type == UE::Dataflow::Private::RotatorType)
-	{
-		return Settings->RotatorPinTypeColor;
-	}
-	else if (Type == UE::Dataflow::Private::ArrayType)
-	{
-		return DataflowSettings->ArrayPinTypeColor;
-	}
-	else if (Type == UE::Dataflow::Private::BoxType)
-	{
-		return DataflowSettings->BoxPinTypeColor;
-	}
-	else if (Type == UE::Dataflow::Private::SphereType)
-	{
-		return DataflowSettings->SpherePinTypeColor;
-	}
-	else if (Type == UE::Dataflow::Private::DataflowAnyTypeType)
+	if (::UE::Dataflow::FPinColorRegistry::Get().IsPinTypeRegistered(Type))
 	{
-		return DataflowSettings->DataflowAnyTypePinTypeColor;
+		return ::UE::Dataflow::FPinColorRegistry::Get().GetPinColor(Type);
 	}
 
 	return Settings->DefaultPinTypeColor;
diff --git a/Engine/Plugins/Experimental/Fracture/Source/FractureEngine/Private/FractureEngineFracturing.cpp b/Engine/Plugins/Experimental/Fracture/Source/FractureEngine/Private/FractureEngineFracturing.cpp
--- a/Engine/Plugins/Experimental/Fracture/Source/FractureEngine/Private/FractureEngineFracturing.cpp
+++ b/Engine/Plugins/Experimental/Fracture/Source/FractureEngine/Private/FractureEngineFracturing.cpp
@@ -12,6 +12,7 @@
 #include "GeometryCollection/Facades/CollectionTransformSelectionFacade.h"
 #include "GeometryCollection/GeometryCollectionClusteringUtility.h"
 #include "FractureEngineMaterials.h"
+#include "Dataflow/DataflowSettings.h"
 
 static void AddAdditionalAttributesIfRequired(FManagedArrayCollection& InOutCollection)
 {
@@ -1361,11 +1362,11 @@ namespace UE::Dataflow::Private
 		}
 	}
 
-	static FLinearColor GetRandomColor(const FRandomStream& InRandomStream)
+	static FLinearColor GetRandomColor(const FRandomStream& InRandomStream, const int32 InColorRangeMin, const int32 InColorRangeMax)
 	{
-		const uint8 R = static_cast<uint8>(InRandomStream.FRandRange(5, 105));
-		const uint8 G = static_cast<uint8>(InRandomStream.FRandRange(5, 105));
-		const uint8 B = static_cast<uint8>(InRandomStream.FRandRange(5, 105));
+		const uint8 R = static_cast<uint8>(InRandomStream.FRandRange(InColorRangeMin, InColorRangeMax));
+		const uint8 G = static_cast<uint8>(InRandomStream.FRandRange(InColorRangeMin, InColorRangeMax));
+		const uint8 B = static_cast<uint8>(InRandomStream.FRandRange(InColorRangeMin, InColorRangeMax));
 
 		return FLinearColor(FColor(R, G, B, 255));
 	}
@@ -1373,22 +1374,20 @@ namespace UE::Dataflow::Private
 
 void FFractureEngineFracturing::InitColors(FManagedArrayCollection& InCollection)
 {
-	//constexpr FLinearColor BoneColor = FLinearColor(0.01, 0.1, 0.1, 1.0);
-	const FLinearColor BoneColor = FLinearColor(FColor(255U, 255U, 255U, 255U));
-	constexpr FLinearColor VertexColor = FLinearColor(1.0, 1.0, 1.0, 1.0);
+	const UDataflowSettings* DataflowSettings = GetDefault<UDataflowSettings>();
 
 	TManagedArray<FLinearColor>& BoneColors = InCollection.ModifyAttribute<FLinearColor>(FGeometryCollection::BoneColorAttribute, FGeometryCollection::TransformGroup);
 	const int32 NumBones = BoneColors.Num();
 	for (int32 BoneIdx = 0; BoneIdx < NumBones; ++BoneIdx)
 	{
-		BoneColors[BoneIdx] = BoneColor;
+		BoneColors[BoneIdx] = DataflowSettings->TransformLevelColors.BlankColor;
 	}
 
 	TManagedArray<FLinearColor>& VertexColors = InCollection.ModifyAttribute<FLinearColor>(FGeometryCollection::ColorAttribute, FGeometryCollection::VerticesGroup);
 	const int32 NumVertices = VertexColors.Num();
 	for (int32 VertexIdx = 0; VertexIdx < NumVertices; ++VertexIdx)
 	{
-		VertexColors[VertexIdx] = VertexColor;
+		VertexColors[VertexIdx] = DataflowSettings->TransformLevelColors.BlankColor;;
 	}
 }
 
@@ -1423,18 +1422,18 @@ void FFractureEngineFracturing::TransferBoneColorToVertexColor(FManagedArrayColl
 void FFractureEngineFracturing::SetBoneColorByParent(
 	FManagedArrayCollection& InCollection,
 	const FRandomStream& InRandomStream,
-	int32 InLevel)
+	int32 InLevel,
+	const int32 InColorRangeMin,
+	const int32 InColorRangeMax)
 {
 	TManagedArray<FLinearColor>& BoneColors = InCollection.ModifyAttribute<FLinearColor>(FGeometryCollection::BoneColorAttribute, FGeometryCollection::TransformGroup);
 	const int32 NumBones = BoneColors.Num();
 	const TManagedArray<int32>& Levels = InCollection.GetAttribute<int32>(FTransformCollection::LevelAttribute, FGeometryCollection::TransformGroup);
 	const TManagedArray<TSet<int32>>& Children = InCollection.GetAttribute<TSet<int32>>(FTransformCollection::ChildrenAttribute, FGeometryCollection::TransformGroup);
 
-	const FLinearColor BlankColor(FColor(255U, 255U, 255U, 255U));
-
 	if (InLevel == 0)
 	{
-		FLinearColor Color = UE::Dataflow::Private::GetRandomColor(InRandomStream);
+		FLinearColor Color = UE::Dataflow::Private::GetRandomColor(InRandomStream, InColorRangeMin, InColorRangeMax);
 
 		for (int32 BoneIdx = 0; BoneIdx < NumBones; ++BoneIdx)
 		{
@@ -1447,7 +1446,7 @@ void FFractureEngineFracturing::SetBoneColorByParent(
 		{
 			if (Levels[BoneIdx] == InLevel - 1)
 			{
-				FLinearColor Color = UE::Dataflow::Private::GetRandomColor(InRandomStream);
+				FLinearColor Color = UE::Dataflow::Private::GetRandomColor(InRandomStream, InColorRangeMin, InColorRangeMax);
 
 				for (int32 ChildBoneIdx : Children[BoneIdx])
 				{
@@ -1460,21 +1459,8 @@ void FFractureEngineFracturing::SetBoneColorByParent(
 
 void FFractureEngineFracturing::SetBoneColorByLevel(FManagedArrayCollection& InCollection, int32 InLevel)
 {
-	const FColor ColorsPerDepth[] =
-	{
-		FColor::Cyan,
-		FColor::Orange,
-		FColor::Emerald,
-		FColor::Yellow,
-		FColor::Purple,
-		FColor::Magenta,
-		FColor::Turquoise,
-		FColor::Emerald,
-		FColor::Silver,
-		FColor::Blue
-	};
-
-	const FLinearColor BlankColor(FColor(255U, 255U, 255U, 255U));
+	const UDataflowSettings* DataflowSettings = GetDefault<UDataflowSettings>();
+	const int32 NumTransformLevelColors = DataflowSettings->TransformLevelColors.LevelColors.Num();
 
 	TManagedArray<FLinearColor>& BoneColors = InCollection.ModifyAttribute<FLinearColor>(FGeometryCollection::BoneColorAttribute, FGeometryCollection::TransformGroup);
 	const int32 NumBones = BoneColors.Num();
@@ -1486,11 +1472,11 @@ void FFractureEngineFracturing::SetBoneColorByLevel(FManagedArrayCollection& InC
 		{
 			if (Levels[BoneIdx] >= InLevel)
 			{
-				BoneColors[BoneIdx] = ColorsPerDepth[InLevel];
+				BoneColors[BoneIdx] = DataflowSettings->TransformLevelColors.LevelColors[InLevel % NumTransformLevelColors];
 			}
 			else
 			{
-				BoneColors[BoneIdx] = BlankColor;
+				BoneColors[BoneIdx] = DataflowSettings->TransformLevelColors.BlankColor;
 			}
 		}
 	}
@@ -1499,21 +1485,23 @@ void FFractureEngineFracturing::SetBoneColorByLevel(FManagedArrayCollection& InC
  void FFractureEngineFracturing::SetBoneColorByCluster(
 	 FManagedArrayCollection& InCollection, 
 	 const FRandomStream& InRandomStream, 
-	 int32 InLevel)
+	 int32 InLevel,
+	 const int32 InColorRangeMin,
+	 const int32 InColorRangeMax)
 {
+	 const UDataflowSettings* DataflowSettings = GetDefault<UDataflowSettings>();
+
 	TArray<FLinearColor> RandomColors;
 
 	FRandomStream Random(1);
 	for (int i = 0; i < 100; i++)
 	{
-		const uint8 R = static_cast<uint8>(Random.FRandRange(5, 105));
-		const uint8 G = static_cast<uint8>(Random.FRandRange(5, 105));
-		const uint8 B = static_cast<uint8>(Random.FRandRange(5, 105));
+		const uint8 R = static_cast<uint8>(Random.FRandRange(InColorRangeMin, InColorRangeMax));
+		const uint8 G = static_cast<uint8>(Random.FRandRange(InColorRangeMin, InColorRangeMax));
+		const uint8 B = static_cast<uint8>(Random.FRandRange(InColorRangeMin, InColorRangeMax));
 		RandomColors.Push(FLinearColor(FColor(R, G, B, 255)));
 	}
 
-	const FLinearColor BlankColor(FColor(255U, 255U, 255U, 255U));
-
 	const TManagedArray<int32>& Parents = InCollection.GetAttribute<int32>("Parent", FGeometryCollection::TransformGroup);
 
 	const TManagedArray<int32>& Levels = InCollection.GetAttribute<int32>(FTransformCollection::LevelAttribute, FGeometryCollection::TransformGroup);
@@ -1543,7 +1531,7 @@ void FFractureEngineFracturing::SetBoneColorByLevel(FManagedArrayCollection& InC
 		}
 		else
 		{
-			BoneColor = BlankColor;
+			BoneColor = DataflowSettings->TransformLevelColors.BlankColor;;
 		}
 
 		BoneColors[BoneIndex] = BoneColor;
@@ -1552,21 +1540,8 @@ void FFractureEngineFracturing::SetBoneColorByLevel(FManagedArrayCollection& InC
 
  void FFractureEngineFracturing::SetBoneColorByLeafLevel(FManagedArrayCollection& InCollection, int32 InLevel)
  {
-	 const FColor ColorsPerDepth[] =
-	 {
-		 FColor::Cyan,
-		 FColor::Orange,
-		 FColor::Emerald,
-		 FColor::Yellow,
-		 FColor::Purple,
-		 FColor::Magenta,
-		 FColor::Turquoise,
-		 FColor::Emerald,
-		 FColor::Silver,
-		 FColor::Blue
-	 };
-
-	 const FLinearColor BlankColor(FColor(255U, 255U, 255U, 255U));
+	 const UDataflowSettings* DataflowSettings = GetDefault<UDataflowSettings>();
+	 const int32 NumTransformLevelColors = DataflowSettings->TransformLevelColors.LevelColors.Num();
 
 	 TManagedArray<FLinearColor>& BoneColors = InCollection.ModifyAttribute<FLinearColor>(FGeometryCollection::BoneColorAttribute, FGeometryCollection::TransformGroup);
 	 const int32 NumBones = BoneColors.Num();
@@ -1577,18 +1552,22 @@ void FFractureEngineFracturing::SetBoneColorByLevel(FManagedArrayCollection& InC
 	 {
 		 if (Levels[BoneIdx] >= InLevel)
 		 {
-			 BoneColors[BoneIdx] = Children[BoneIdx].Num() == 0 ? ColorsPerDepth[Levels[BoneIdx]] : FColor::Black;
+			 BoneColors[BoneIdx] = Children[BoneIdx].Num() == 0 ? DataflowSettings->TransformLevelColors.LevelColors[Levels[BoneIdx] % NumTransformLevelColors] : FColor::Black;
 		 }
 		 else
 		 {
-			 BoneColors[BoneIdx] = BlankColor;
+			 BoneColors[BoneIdx] = DataflowSettings->TransformLevelColors.BlankColor;
 		 }
 	 }
  }
 
- void FFractureEngineFracturing::SetBoneColorByLeaf(FManagedArrayCollection& InCollection, const FRandomStream& InRandomStream, int32 InLevel)
+ void FFractureEngineFracturing::SetBoneColorByLeaf(FManagedArrayCollection& InCollection, 
+	 const FRandomStream& InRandomStream, 
+	 int32 InLevel, 
+	 const int32 InColorRangeMin,
+	 const int32 InColorRangeMax)
  {
-	 const FLinearColor BlankColor(FColor(255U, 255U, 255U, 255U));
+	 const UDataflowSettings* DataflowSettings = GetDefault<UDataflowSettings>();
 
 	 TManagedArray<FLinearColor>& BoneColors = InCollection.ModifyAttribute<FLinearColor>(FGeometryCollection::BoneColorAttribute, FGeometryCollection::TransformGroup);
 	 const TManagedArray<int32>& Levels = InCollection.GetAttribute<int32>(FTransformCollection::LevelAttribute, FGeometryCollection::TransformGroup);
@@ -1598,11 +1577,11 @@ void FFractureEngineFracturing::SetBoneColorByLevel(FManagedArrayCollection& InC
 	 {
 		 if (Levels[BoneIdx] >= InLevel)
 		 {
-			 BoneColors[BoneIdx] = UE::Dataflow::Private::GetRandomColor(InRandomStream);
+			 BoneColors[BoneIdx] = UE::Dataflow::Private::GetRandomColor(InRandomStream, InColorRangeMin, InColorRangeMax);
 		 }
 		 else
 		 {
-			 BoneColors[BoneIdx] = BlankColor;
+			 BoneColors[BoneIdx] = DataflowSettings->TransformLevelColors.BlankColor;
 		 }
 	 }
  }
diff --git a/Engine/Plugins/Experimental/Fracture/Source/FractureEngine/Public/FractureEngineFracturing.h b/Engine/Plugins/Experimental/Fracture/Source/FractureEngine/Public/FractureEngineFracturing.h
--- a/Engine/Plugins/Experimental/Fracture/Source/FractureEngine/Public/FractureEngineFracturing.h
+++ b/Engine/Plugins/Experimental/Fracture/Source/FractureEngine/Public/FractureEngineFracturing.h
@@ -211,7 +211,9 @@ public:
 	static void SetBoneColorByParent(
 		FManagedArrayCollection& InCollection,
 		const FRandomStream& InRandomStream,
-		int32 InLevel);
+		int32 InLevel,
+		const int32 InColorRangeMin = 40,
+		const int32 InColorRangeMax = 190);
 
 	static void SetBoneColorByLevel(
 		FManagedArrayCollection& InCollection, 
@@ -220,7 +222,9 @@ public:
 	static void SetBoneColorByCluster(
 		FManagedArrayCollection& InCollection,
 		const FRandomStream& InRandomStream,
-		int32 InLevel);
+		int32 InLevel,
+		const int32 InColorRangeMin,
+		const int32 InColorRangeMax);
 
 	static void SetBoneColorByLeafLevel(
 		FManagedArrayCollection& InCollection, 
@@ -229,7 +233,9 @@ public:
 	static void SetBoneColorByLeaf(
 		FManagedArrayCollection& InCollection, 
 		const FRandomStream& InRandomStream,
-		int32 InLevel);
+		int32 InLevel,
+		const int32 InColorRangeMin,
+		const int32 InColorRangeMax);
 
 	static void SetBoneColorByAttr(
 		FManagedArrayCollection& InCollection,
diff --git a/Engine/Plugins/Experimental/GeometryCollectionPlugin/Source/GeometryCollectionNodes/Private/Dataflow/DataflowNodeAndPinTypeColors.cpp b/Engine/Plugins/Experimental/GeometryCollectionPlugin/Source/GeometryCollectionNodes/Private/Dataflow/DataflowNodeAndPinTypeColors.cpp
new file mode 100644
--- /dev/null
+++ b/Engine/Plugins/Experimental/GeometryCollectionPlugin/Source/GeometryCollectionNodes/Private/Dataflow/DataflowNodeAndPinTypeColors.cpp
@@ -0,0 +1,62 @@
+// Copyright Epic Games, Inc. All Rights Reserved.
+
+#include "Dataflow/DataflowNodeAndPinTypeColors.h"
+#include "Dataflow/DataflowNode.h"
+#include "Math/Color.h" 
+#include "Dataflow/DataflowNodeColorsRegistry.h"
+
+namespace UE::Dataflow
+{
+	void RegisterGeometryCollectionNodesColors()
+	{
+		static const FLinearColor CDefaultNodeBodyTintColor = FLinearColor(0.f, 0.f, 0.f, 0.5f);
+
+		// Node colors
+		DATAFLOW_NODE_REGISTER_CREATION_FACTORY_NODE_COLORS_BY_CATEGORY("GeometryCollection|Fracture", FLinearColor(1.f, 1.f, 0.8f), CDefaultNodeBodyTintColor);
+		DATAFLOW_NODE_REGISTER_CREATION_FACTORY_NODE_COLORS_BY_CATEGORY("Terminal", FLinearColor(0.f, 0.f, 0.f), CDefaultNodeBodyTintColor);
+		DATAFLOW_NODE_REGISTER_CREATION_FACTORY_NODE_COLORS_BY_CATEGORY("GeometryCollection|Cluster", FLinearColor(.25f, 0.45f, 0.8f), CDefaultNodeBodyTintColor);
+		DATAFLOW_NODE_REGISTER_CREATION_FACTORY_NODE_COLORS_BY_CATEGORY("Generators", FLinearColor(.4f, 0.8f, 0.f), CDefaultNodeBodyTintColor);
+		DATAFLOW_NODE_REGISTER_CREATION_FACTORY_NODE_COLORS_BY_CATEGORY("GeometryCollection|Development", FLinearColor(1.f, 0.f, 0.f), CDefaultNodeBodyTintColor);
+		DATAFLOW_NODE_REGISTER_CREATION_FACTORY_NODE_COLORS_BY_CATEGORY("GeometryCollection|Mesh", FLinearColor(1.f, 0.16f, 0.05f), CDefaultNodeBodyTintColor);
+		DATAFLOW_NODE_REGISTER_CREATION_FACTORY_NODE_COLORS_BY_CATEGORY("SphereCovering", FLinearColor(1.f, 0.16f, 0.05f), CDefaultNodeBodyTintColor);
+		DATAFLOW_NODE_REGISTER_CREATION_FACTORY_NODE_COLORS_BY_CATEGORY("GeometryCollection|Edit", FLinearColor(0.f, 1.f, 0.05f), CDefaultNodeBodyTintColor);
+		DATAFLOW_NODE_REGISTER_CREATION_FACTORY_NODE_COLORS_BY_CATEGORY("Fields", FLinearColor(.0f, 0.8f, 1.f), CDefaultNodeBodyTintColor);
+		DATAFLOW_NODE_REGISTER_CREATION_FACTORY_NODE_COLORS_BY_CATEGORY("Generators", FLinearColor(.4f, 0.8f, 0.f), CDefaultNodeBodyTintColor);
+		DATAFLOW_NODE_REGISTER_CREATION_FACTORY_NODE_COLORS_BY_CATEGORY("Materials", FLinearColor(FColor::Purple), CDefaultNodeBodyTintColor);
+		DATAFLOW_NODE_REGISTER_CREATION_FACTORY_NODE_COLORS_BY_CATEGORY("Math", FLinearColor(0.f, 0.4f, 0.8f), CDefaultNodeBodyTintColor);
+		DATAFLOW_NODE_REGISTER_CREATION_FACTORY_NODE_COLORS_BY_CATEGORY("Mesh", FLinearColor(1.f, 0.16f, 0.05f), CDefaultNodeBodyTintColor);
+		DATAFLOW_NODE_REGISTER_CREATION_FACTORY_NODE_COLORS_BY_CATEGORY("GeometryCollection", FLinearColor(0.55f, 0.45f, 1.0f), CDefaultNodeBodyTintColor);
+		DATAFLOW_NODE_REGISTER_CREATION_FACTORY_NODE_COLORS_BY_CATEGORY("Development", FLinearColor(1.f, 0.f, 0.f), CDefaultNodeBodyTintColor);
+		DATAFLOW_NODE_REGISTER_CREATION_FACTORY_NODE_COLORS_BY_CATEGORY("Utilities|String", FLinearColor(0.5f, 0.f, 0.5f), CDefaultNodeBodyTintColor);
+		DATAFLOW_NODE_REGISTER_CREATION_FACTORY_NODE_COLORS_BY_CATEGORY("Fracture", FLinearColor(1.f, 1.f, 0.8f), CDefaultNodeBodyTintColor);
+		DATAFLOW_NODE_REGISTER_CREATION_FACTORY_NODE_COLORS_BY_CATEGORY("Utilities", FLinearColor(1.f, 1.f, 0.f), CDefaultNodeBodyTintColor);
+		DATAFLOW_NODE_REGISTER_CREATION_FACTORY_NODE_COLORS_BY_CATEGORY("Override", FLinearColor(1.f, 0.4f, 0.4f), CDefaultNodeBodyTintColor);
+		DATAFLOW_NODE_REGISTER_CREATION_FACTORY_NODE_COLORS_BY_CATEGORY("GeometryCollection|Sampling", FLinearColor(.1f, 1.f, 0.6f), CDefaultNodeBodyTintColor);
+		DATAFLOW_NODE_REGISTER_CREATION_FACTORY_NODE_COLORS_BY_CATEGORY("GeometryCollection|Selection", FLinearColor(1.f, 1.f, 0.05f), CDefaultNodeBodyTintColor);
+		DATAFLOW_NODE_REGISTER_CREATION_FACTORY_NODE_COLORS_BY_CATEGORY("GeometryCollection|Texture", FLinearColor(.1f, 1.f, 0.6f), CDefaultNodeBodyTintColor);
+		DATAFLOW_NODE_REGISTER_CREATION_FACTORY_NODE_COLORS_BY_CATEGORY("GeometryCollection|UV", FLinearColor(.1f, 1.f, 0.6f), CDefaultNodeBodyTintColor);
+
+		// PinType colors
+		DATAFLOW_NODE_REGISTER_CREATION_FACTORY_PIN_COLORS_BY_TYPE("FManagedArrayCollection", FLinearColor(0.353393f, 0.454175f, 1.0f, 1.0f));
+		DATAFLOW_NODE_REGISTER_CREATION_FACTORY_PIN_COLORS_BY_TYPE("float", FLinearColor(0.357667f, 1.f, 0.06f));
+		DATAFLOW_NODE_REGISTER_CREATION_FACTORY_PIN_COLORS_BY_TYPE("double", FLinearColor(0.039216f, 0.666667f, 0.f));
+		DATAFLOW_NODE_REGISTER_CREATION_FACTORY_PIN_COLORS_BY_TYPE("byte", FLinearColor(0.f, 0.16f, 0.13127f));
+		DATAFLOW_NODE_REGISTER_CREATION_FACTORY_PIN_COLORS_BY_TYPE("int32", FLinearColor(0.013575f, 0.77f, 0.429609f));
+		DATAFLOW_NODE_REGISTER_CREATION_FACTORY_PIN_COLORS_BY_TYPE("bool", FLinearColor(0.3f, 0.f, 0.f, 1.f));
+		DATAFLOW_NODE_REGISTER_CREATION_FACTORY_PIN_COLORS_BY_TYPE("FString", FLinearColor(1.f, 0.f, 0.660537f));
+		DATAFLOW_NODE_REGISTER_CREATION_FACTORY_PIN_COLORS_BY_TYPE("FName", FLinearColor(0.607717f, 0.224984f, 1.f));
+		DATAFLOW_NODE_REGISTER_CREATION_FACTORY_PIN_COLORS_BY_TYPE("FText", FLinearColor(0.8f, 0.2f, 0.4f));
+		DATAFLOW_NODE_REGISTER_CREATION_FACTORY_PIN_COLORS_BY_TYPE("FVector", FLinearColor(1.f, 0.591255f, 0.016512f));
+		DATAFLOW_NODE_REGISTER_CREATION_FACTORY_PIN_COLORS_BY_TYPE("FTransform", FLinearColor(1.f, 0.172585f, 0.f));
+		DATAFLOW_NODE_REGISTER_CREATION_FACTORY_PIN_COLORS_BY_TYPE("FRotator", FLinearColor(0.353393f, 0.454175f, 1.f));
+		DATAFLOW_NODE_REGISTER_CREATION_FACTORY_PIN_COLORS_BY_TYPE("TArray<FString>", FLinearColor(1.0f, 0.172585f, 0.0f, 1.0f));
+		DATAFLOW_NODE_REGISTER_CREATION_FACTORY_PIN_COLORS_BY_TYPE("FBox", FLinearColor(0.013575f, 0.770000f, 0.429609f, 1.0f));
+		DATAFLOW_NODE_REGISTER_CREATION_FACTORY_PIN_COLORS_BY_TYPE("FSphere", FLinearColor(0.2f, 0.6f, 1.f, 1.0f));
+		DATAFLOW_NODE_REGISTER_CREATION_FACTORY_PIN_COLORS_BY_TYPE("FDataflowAnyType", FLinearColor(0.3f, 0.3f, 0.3f, 1.0f));
+		DATAFLOW_NODE_REGISTER_CREATION_FACTORY_PIN_COLORS_BY_TYPE("FDataflowTransformSelection", FLinearColor(0.f, 0.9f, 0.0f, 1.0f));
+		DATAFLOW_NODE_REGISTER_CREATION_FACTORY_PIN_COLORS_BY_TYPE("TArray<TObjectPtr<UMaterial>>", FLinearColor(0.f, 0.2f, 0.0f, 1.0f));
+		DATAFLOW_NODE_REGISTER_CREATION_FACTORY_PIN_COLORS_BY_TYPE("TArray<TObjectPtr<UMaterialInterface>>", FLinearColor(0.f, 0.4f, 0.0f, 1.0f));
+		DATAFLOW_NODE_REGISTER_CREATION_FACTORY_PIN_COLORS_BY_TYPE("TArray<FGeometryCollectionAutoInstanceMesh>", FLinearColor(0.f, 0.6f, 0.0f, 1.0f));
+		DATAFLOW_NODE_REGISTER_CREATION_FACTORY_PIN_COLORS_BY_TYPE("TObjectPtr<UStaticMesh>", FLinearColor(0.f, 0.8f, 0.0f, 1.0f));
+	}
+}
diff --git a/Engine/Plugins/Experimental/GeometryCollectionPlugin/Source/GeometryCollectionNodes/Private/Dataflow/GeometryCollectionArrayNodes.cpp b/Engine/Plugins/Experimental/GeometryCollectionPlugin/Source/GeometryCollectionNodes/Private/Dataflow/GeometryCollectionArrayNodes.cpp
--- a/Engine/Plugins/Experimental/GeometryCollectionPlugin/Source/GeometryCollectionNodes/Private/Dataflow/GeometryCollectionArrayNodes.cpp
+++ b/Engine/Plugins/Experimental/GeometryCollectionPlugin/Source/GeometryCollectionNodes/Private/Dataflow/GeometryCollectionArrayNodes.cpp
@@ -11,8 +11,6 @@ namespace UE::Dataflow
 
 	void GeometryCollectionArrayNodes()
 	{
-		static const FLinearColor CDefaultNodeBodyTintColor = FLinearColor(0.f, 0.f, 0.f, 0.5f);
-
 		DATAFLOW_NODE_REGISTER_CREATION_FACTORY(FGetFloatArrayElementDataflowNode);
 		DATAFLOW_NODE_REGISTER_CREATION_FACTORY(FFloatArrayToIntArrayDataflowNode);
 		DATAFLOW_NODE_REGISTER_CREATION_FACTORY(FGetArrayElementDataflowNode);
diff --git a/Engine/Plugins/Experimental/GeometryCollectionPlugin/Source/GeometryCollectionNodes/Private/Dataflow/GeometryCollectionAssetNodes.cpp b/Engine/Plugins/Experimental/GeometryCollectionPlugin/Source/GeometryCollectionNodes/Private/Dataflow/GeometryCollectionAssetNodes.cpp
--- a/Engine/Plugins/Experimental/GeometryCollectionPlugin/Source/GeometryCollectionNodes/Private/Dataflow/GeometryCollectionAssetNodes.cpp
+++ b/Engine/Plugins/Experimental/GeometryCollectionPlugin/Source/GeometryCollectionNodes/Private/Dataflow/GeometryCollectionAssetNodes.cpp
@@ -17,8 +17,6 @@ namespace UE::Dataflow
 {
 	void GeometryCollectionEngineAssetNodes()
 	{
-		static const FLinearColor CDefaultNodeBodyTintColor = FLinearColor(0.f, 0.f, 0.f, 0.5f);
-
 		DATAFLOW_NODE_REGISTER_CREATION_FACTORY(FGeometryCollectionTerminalDataflowNode_v2);
 		DATAFLOW_NODE_REGISTER_CREATION_FACTORY(FGetGeometryCollectionAssetDataflowNode);
 		DATAFLOW_NODE_REGISTER_CREATION_FACTORY(FGetGeometryCollectionSourcesDataflowNode);
@@ -28,9 +26,6 @@ namespace UE::Dataflow
 
 		// deprecated nodes (need to stay registered)
 		DATAFLOW_NODE_REGISTER_CREATION_FACTORY(FGeometryCollectionTerminalDataflowNode);
-
-		// Terminal
-		DATAFLOW_NODE_REGISTER_CREATION_FACTORY_NODE_COLORS_BY_CATEGORY("Terminal", FLinearColor(0.f, 0.f, 0.f), CDefaultNodeBodyTintColor);
 	}
 }
 
diff --git a/Engine/Plugins/Experimental/GeometryCollectionPlugin/Source/GeometryCollectionNodes/Private/Dataflow/GeometryCollectionClusteringNodes.cpp b/Engine/Plugins/Experimental/GeometryCollectionPlugin/Source/GeometryCollectionNodes/Private/Dataflow/GeometryCollectionClusteringNodes.cpp
--- a/Engine/Plugins/Experimental/GeometryCollectionPlugin/Source/GeometryCollectionNodes/Private/Dataflow/GeometryCollectionClusteringNodes.cpp
+++ b/Engine/Plugins/Experimental/GeometryCollectionPlugin/Source/GeometryCollectionNodes/Private/Dataflow/GeometryCollectionClusteringNodes.cpp
@@ -38,8 +38,6 @@ namespace UE::Dataflow
 
 	void GeometryCollectionClusteringNodes()
 	{
-		static const FLinearColor CDefaultNodeBodyTintColor = FLinearColor(0.f, 0.f, 0.f, 0.5f);
-
 		DATAFLOW_NODE_REGISTER_CREATION_FACTORY(FAutoClusterDataflowNode);
 		DATAFLOW_NODE_REGISTER_CREATION_FACTORY(FClusterFlattenDataflowNode);
 		DATAFLOW_NODE_REGISTER_CREATION_FACTORY(FClusterUnclusterDataflowNode);
@@ -48,9 +46,6 @@ namespace UE::Dataflow
 		DATAFLOW_NODE_REGISTER_CREATION_FACTORY(FClusterMergeDataflowNode);
 		DATAFLOW_NODE_REGISTER_CREATION_FACTORY(FClusterIsolatedRootsDataflowNode);
 		DATAFLOW_NODE_REGISTER_CREATION_FACTORY(FClusterMagnetDataflowNode);
-
-		// GeometryCollection|Cluster
-		DATAFLOW_NODE_REGISTER_CREATION_FACTORY_NODE_COLORS_BY_CATEGORY("GeometryCollection|Cluster", FLinearColor(.25f, 0.45f, 0.8f), CDefaultNodeBodyTintColor);
 	}
 }
 
diff --git a/Engine/Plugins/Experimental/GeometryCollectionPlugin/Source/GeometryCollectionNodes/Private/Dataflow/GeometryCollectionConversionNodes.cpp b/Engine/Plugins/Experimental/GeometryCollectionPlugin/Source/GeometryCollectionNodes/Private/Dataflow/GeometryCollectionConversionNodes.cpp
--- a/Engine/Plugins/Experimental/GeometryCollectionPlugin/Source/GeometryCollectionNodes/Private/Dataflow/GeometryCollectionConversionNodes.cpp
+++ b/Engine/Plugins/Experimental/GeometryCollectionPlugin/Source/GeometryCollectionNodes/Private/Dataflow/GeometryCollectionConversionNodes.cpp
@@ -9,8 +9,6 @@ namespace UE::Dataflow
 {
 	void GeometryCollectionConversionNodes()
 	{
-		static const FLinearColor CDefaultNodeBodyTintColor = FLinearColor(0.f, 0.f, 0.f, 0.5f);
-
 		DATAFLOW_NODE_REGISTER_CREATION_FACTORY(FVectorToStringDataflowNode);
 		DATAFLOW_NODE_REGISTER_CREATION_FACTORY(FFloatToStringDataflowNode);
 		DATAFLOW_NODE_REGISTER_CREATION_FACTORY(FIntToStringDataflowNode);
@@ -21,10 +19,6 @@ namespace UE::Dataflow
 		DATAFLOW_NODE_REGISTER_CREATION_FACTORY(FFloatToIntDataflowNode);
 		DATAFLOW_NODE_REGISTER_CREATION_FACTORY(FIntToBoolDataflowNode);
 		DATAFLOW_NODE_REGISTER_CREATION_FACTORY(FBoolToIntDataflowNode);
-
-
-		// Generators
-		DATAFLOW_NODE_REGISTER_CREATION_FACTORY_NODE_COLORS_BY_CATEGORY("Generators", FLinearColor(.4f, 0.8f, 0.f), CDefaultNodeBodyTintColor);
 	}
 }
 
diff --git a/Engine/Plugins/Experimental/GeometryCollectionPlugin/Source/GeometryCollectionNodes/Private/Dataflow/GeometryCollectionDebugNodes.cpp b/Engine/Plugins/Experimental/GeometryCollectionPlugin/Source/GeometryCollectionNodes/Private/Dataflow/GeometryCollectionDebugNodes.cpp
--- a/Engine/Plugins/Experimental/GeometryCollectionPlugin/Source/GeometryCollectionNodes/Private/Dataflow/GeometryCollectionDebugNodes.cpp
+++ b/Engine/Plugins/Experimental/GeometryCollectionPlugin/Source/GeometryCollectionNodes/Private/Dataflow/GeometryCollectionDebugNodes.cpp
@@ -18,18 +18,11 @@ namespace UE::Dataflow
 {
 	void GeometryCollectionDebugNodes()
 	{
-		static const FLinearColor CDefaultNodeBodyTintColor = FLinearColor(0.f, 0.f, 0.f, 0.5f);
-
 		DATAFLOW_NODE_REGISTER_CREATION_FACTORY(FConvexHullToMeshDataflowNode);
 		DATAFLOW_NODE_REGISTER_CREATION_FACTORY(FSphereCoveringToMeshDataflowNode);
 		DATAFLOW_NODE_REGISTER_CREATION_FACTORY(FMeshToOBJStringDebugDataflowNode);
 		DATAFLOW_NODE_REGISTER_CREATION_FACTORY(FSphereCoveringCountSpheresNode);
 		DATAFLOW_NODE_REGISTER_CREATION_FACTORY(FWriteStringToFile);
-
-
-		DATAFLOW_NODE_REGISTER_CREATION_FACTORY_NODE_COLORS_BY_CATEGORY("GeometryCollection|Development", FLinearColor(1.f, 0.f, 0.f), CDefaultNodeBodyTintColor);
-		DATAFLOW_NODE_REGISTER_CREATION_FACTORY_NODE_COLORS_BY_CATEGORY("GeometryCollection|Mesh", FLinearColor(1.f, 0.16f, 0.05f), CDefaultNodeBodyTintColor);
-		DATAFLOW_NODE_REGISTER_CREATION_FACTORY_NODE_COLORS_BY_CATEGORY("SphereCovering", FLinearColor(1.f, 0.16f, 0.05f), CDefaultNodeBodyTintColor);
 	}
 }
 
diff --git a/Engine/Plugins/Experimental/GeometryCollectionPlugin/Source/GeometryCollectionNodes/Private/Dataflow/GeometryCollectionEditNodes.cpp b/Engine/Plugins/Experimental/GeometryCollectionPlugin/Source/GeometryCollectionNodes/Private/Dataflow/GeometryCollectionEditNodes.cpp
--- a/Engine/Plugins/Experimental/GeometryCollectionPlugin/Source/GeometryCollectionNodes/Private/Dataflow/GeometryCollectionEditNodes.cpp
+++ b/Engine/Plugins/Experimental/GeometryCollectionPlugin/Source/GeometryCollectionNodes/Private/Dataflow/GeometryCollectionEditNodes.cpp
@@ -35,14 +35,9 @@ namespace UE::Dataflow
 
 	void GeometryCollectionEditNodes()
 	{
-		static const FLinearColor CDefaultNodeBodyTintColor = FLinearColor(0.f, 0.f, 0.f, 0.5f);
-
 		DATAFLOW_NODE_REGISTER_CREATION_FACTORY(FPruneInCollectionDataflowNode);
 		DATAFLOW_NODE_REGISTER_CREATION_FACTORY(FSetVisibilityInCollectionDataflowNode);
 		DATAFLOW_NODE_REGISTER_CREATION_FACTORY(FMergeInCollectionDataflowNode);
-
-		// GeometryCollection|Edit
-		DATAFLOW_NODE_REGISTER_CREATION_FACTORY_NODE_COLORS_BY_CATEGORY("GeometryCollection|Edit", FLinearColor(0.f, 1.f, 0.05f), CDefaultNodeBodyTintColor);
 	}
 }
 
diff --git a/Engine/Plugins/Experimental/GeometryCollectionPlugin/Source/GeometryCollectionNodes/Private/Dataflow/GeometryCollectionEmbedNodes.cpp b/Engine/Plugins/Experimental/GeometryCollectionPlugin/Source/GeometryCollectionNodes/Private/Dataflow/GeometryCollectionEmbedNodes.cpp
--- a/Engine/Plugins/Experimental/GeometryCollectionPlugin/Source/GeometryCollectionNodes/Private/Dataflow/GeometryCollectionEmbedNodes.cpp
+++ b/Engine/Plugins/Experimental/GeometryCollectionPlugin/Source/GeometryCollectionNodes/Private/Dataflow/GeometryCollectionEmbedNodes.cpp
@@ -11,8 +11,6 @@ namespace UE::Dataflow
 
 	void GeometryCollectionEmbedNodes()
 	{
-		static const FLinearColor CDefaultNodeBodyTintColor = FLinearColor(0.f, 0.f, 0.f, 0.5f);
-
 	}
 }
 
diff --git a/Engine/Plugins/Experimental/GeometryCollectionPlugin/Source/GeometryCollectionNodes/Private/Dataflow/GeometryCollectionFieldNodes.cpp b/Engine/Plugins/Experimental/GeometryCollectionPlugin/Source/GeometryCollectionNodes/Private/Dataflow/GeometryCollectionFieldNodes.cpp
--- a/Engine/Plugins/Experimental/GeometryCollectionPlugin/Source/GeometryCollectionNodes/Private/Dataflow/GeometryCollectionFieldNodes.cpp
+++ b/Engine/Plugins/Experimental/GeometryCollectionPlugin/Source/GeometryCollectionNodes/Private/Dataflow/GeometryCollectionFieldNodes.cpp
@@ -51,8 +51,6 @@ namespace UE::Dataflow
 {
 	void GeometryCollectionFieldNodes()
 	{
-		static const FLinearColor CDefaultNodeBodyTintColor = FLinearColor(0.f, 0.f, 0.f, 0.5f);
-
 		DATAFLOW_NODE_REGISTER_CREATION_FACTORY(FRadialFalloffFieldDataflowNode);
 		DATAFLOW_NODE_REGISTER_CREATION_FACTORY(FBoxFalloffFieldDataflowNode);
 		DATAFLOW_NODE_REGISTER_CREATION_FACTORY(FPlaneFalloffFieldDataflowNode);
@@ -67,9 +65,6 @@ namespace UE::Dataflow
 		DATAFLOW_NODE_REGISTER_CREATION_FACTORY(FSumScalarFieldDataflowNode);
 		DATAFLOW_NODE_REGISTER_CREATION_FACTORY(FSumVectorFieldDataflowNode);
 		DATAFLOW_NODE_REGISTER_CREATION_FACTORY(FFieldMakeDenseFloatArrayDataflowNode);
-
-		// Field nodes
-		DATAFLOW_NODE_REGISTER_CREATION_FACTORY_NODE_COLORS_BY_CATEGORY("Fields", FLinearColor(.0f, 0.8f, 1.f), CDefaultNodeBodyTintColor);
 	}
 }
 
diff --git a/Engine/Plugins/Experimental/GeometryCollectionPlugin/Source/GeometryCollectionNodes/Private/Dataflow/GeometryCollectionFracturingNodes.cpp b/Engine/Plugins/Experimental/GeometryCollectionPlugin/Source/GeometryCollectionNodes/Private/Dataflow/GeometryCollectionFracturingNodes.cpp
--- a/Engine/Plugins/Experimental/GeometryCollectionPlugin/Source/GeometryCollectionNodes/Private/Dataflow/GeometryCollectionFracturingNodes.cpp
+++ b/Engine/Plugins/Experimental/GeometryCollectionPlugin/Source/GeometryCollectionNodes/Private/Dataflow/GeometryCollectionFracturingNodes.cpp
@@ -44,8 +44,6 @@ namespace UE::Dataflow
 
 	void GeometryCollectionFracturingNodes()
 	{
-		static const FLinearColor CDefaultNodeBodyTintColor = FLinearColor(0.f, 0.f, 0.f, 0.5f);
-
 		DATAFLOW_NODE_REGISTER_CREATION_FACTORY(FUniformScatterPointsDataflowNode);
 		DATAFLOW_NODE_REGISTER_CREATION_FACTORY(FUniformScatterPointsDataflowNode_v2);
 		DATAFLOW_NODE_REGISTER_CREATION_FACTORY(FRadialScatterPointsDataflowNode);
@@ -63,9 +61,6 @@ namespace UE::Dataflow
 		DATAFLOW_NODE_REGISTER_CREATION_FACTORY(FVisualizeFractureDataflowNode);
 //		Commented out until we decide how to make generic data setter nodes
 //		DATAFLOW_NODE_REGISTER_CREATION_FACTORY(FSetFloatAttributeDataflowNode);
-
-		// GeometryCollection|Fracture
-		DATAFLOW_NODE_REGISTER_CREATION_FACTORY_NODE_COLORS_BY_CATEGORY("GeometryCollection|Fracture", FLinearColor(1.f, 1.f, 0.8f), CDefaultNodeBodyTintColor);
 	}
 }
 
@@ -1182,7 +1177,7 @@ void FVisualizeFractureDataflowNode::Evaluate(UE::Dataflow::FContext& Context, c
 				if (ColoringType == EDataflowVisualizeFractureColoringType::ColorByParent && 
 					InCollection.HasAttribute(FTransformCollection::LevelAttribute, FGeometryCollection::TransformGroup))
 				{
-					FFractureEngineFracturing::SetBoneColorByParent(InCollection, RandomStream, InLevel);
+					FFractureEngineFracturing::SetBoneColorByParent(InCollection, RandomStream, InLevel, RandomColorRangeMin, RandomColorRangeMax);
 				}
 				else if (ColoringType == EDataflowVisualizeFractureColoringType::ColorByLevel &&
 					InCollection.HasAttribute(FTransformCollection::LevelAttribute, FGeometryCollection::TransformGroup))
@@ -1193,7 +1188,7 @@ void FVisualizeFractureDataflowNode::Evaluate(UE::Dataflow::FContext& Context, c
 					InCollection.HasAttribute(FTransformCollection::LevelAttribute, FGeometryCollection::TransformGroup))
 				{
 					// This what the Geometry Tools uses
-					FFractureEngineFracturing::SetBoneColorByCluster(InCollection, RandomStream, InLevel);
+					FFractureEngineFracturing::SetBoneColorByCluster(InCollection, RandomStream, InLevel, RandomColorRangeMin, RandomColorRangeMax);
 				}
 				else if (ColoringType == EDataflowVisualizeFractureColoringType::ColorByLeafLevel &&
 					InCollection.HasAttribute(FTransformCollection::LevelAttribute, FGeometryCollection::TransformGroup))
@@ -1202,7 +1197,7 @@ void FVisualizeFractureDataflowNode::Evaluate(UE::Dataflow::FContext& Context, c
 				}
 				else if (ColoringType == EDataflowVisualizeFractureColoringType::ColorByLeaf)
 				{
-					FFractureEngineFracturing::SetBoneColorByLeaf(InCollection, RandomStream, InLevel);
+					FFractureEngineFracturing::SetBoneColorByLeaf(InCollection, RandomStream, InLevel, RandomColorRangeMin, RandomColorRangeMax);
 				}
 				else if (ColoringType == EDataflowVisualizeFractureColoringType::ColorByAttr)
 				{
diff --git a/Engine/Plugins/Experimental/GeometryCollectionPlugin/Source/GeometryCollectionNodes/Private/Dataflow/GeometryCollectionMakeNodes.cpp b/Engine/Plugins/Experimental/GeometryCollectionPlugin/Source/GeometryCollectionNodes/Private/Dataflow/GeometryCollectionMakeNodes.cpp
--- a/Engine/Plugins/Experimental/GeometryCollectionPlugin/Source/GeometryCollectionNodes/Private/Dataflow/GeometryCollectionMakeNodes.cpp
+++ b/Engine/Plugins/Experimental/GeometryCollectionPlugin/Source/GeometryCollectionNodes/Private/Dataflow/GeometryCollectionMakeNodes.cpp
@@ -44,8 +44,6 @@ namespace UE::Dataflow
 {
 	void GeometryCollectionMakeNodes()
 	{
-		static const FLinearColor CDefaultNodeBodyTintColor = FLinearColor(0.f, 0.f, 0.f, 0.5f);
-
 		DATAFLOW_NODE_REGISTER_CREATION_FACTORY(FMakeLiteralStringDataflowNode);
 		DATAFLOW_NODE_REGISTER_CREATION_FACTORY(FMakePointsDataflowNode);
 		DATAFLOW_NODE_REGISTER_CREATION_FACTORY(FMakeBoxDataflowNode);
@@ -58,9 +56,6 @@ namespace UE::Dataflow
 		DATAFLOW_NODE_REGISTER_CREATION_FACTORY(FMakeQuaternionDataflowNode);
 		DATAFLOW_NODE_REGISTER_CREATION_FACTORY(FMakeFloatArrayDataflowNode);
 		DATAFLOW_NODE_REGISTER_CREATION_FACTORY(FMakeCollectionDataflowNode);
-
-		// Generators
-		DATAFLOW_NODE_REGISTER_CREATION_FACTORY_NODE_COLORS_BY_CATEGORY("Generators", FLinearColor(.4f, 0.8f, 0.f), CDefaultNodeBodyTintColor);
 	}
 }
 
diff --git a/Engine/Plugins/Experimental/GeometryCollectionPlugin/Source/GeometryCollectionNodes/Private/Dataflow/GeometryCollectionMaterialInterfaceNodes.cpp b/Engine/Plugins/Experimental/GeometryCollectionPlugin/Source/GeometryCollectionNodes/Private/Dataflow/GeometryCollectionMaterialInterfaceNodes.cpp
--- a/Engine/Plugins/Experimental/GeometryCollectionPlugin/Source/GeometryCollectionNodes/Private/Dataflow/GeometryCollectionMaterialInterfaceNodes.cpp
+++ b/Engine/Plugins/Experimental/GeometryCollectionPlugin/Source/GeometryCollectionNodes/Private/Dataflow/GeometryCollectionMaterialInterfaceNodes.cpp
@@ -8,8 +8,6 @@ namespace UE::Dataflow
 {
 	void RegisterGeometryCollectionMaterialInterfaceNodes()
 	{
-		static const FLinearColor CDefaultNodeBodyTintColor = FLinearColor(0.f, 0.f, 0.f, 0.5f);
-
 		DATAFLOW_NODE_REGISTER_CREATION_FACTORY(FMakeMaterialInterfaceArrayDataflowNode);
 		DATAFLOW_NODE_REGISTER_CREATION_FACTORY(FGetMaterialInterfaceAssetDataflowNode);
 		DATAFLOW_NODE_REGISTER_CREATION_FACTORY(FGetMaterialInterfaceArraySizeDataflowNode);
@@ -17,9 +15,6 @@ namespace UE::Dataflow
 		DATAFLOW_NODE_REGISTER_CREATION_FACTORY(FSetIntoMaterialInterfaceArrayDataflowNode);
 		DATAFLOW_NODE_REGISTER_CREATION_FACTORY(FAddToMaterialInterfaceArrayDataflowNode);
 		DATAFLOW_NODE_REGISTER_CREATION_FACTORY(FAssignMaterialInterfaceToCollectionDataflowNode);
-
-		// Material
-		DATAFLOW_NODE_REGISTER_CREATION_FACTORY_NODE_COLORS_BY_CATEGORY("Materials", FLinearColor(FColor::Purple), CDefaultNodeBodyTintColor);
 	}
 }
 
diff --git a/Engine/Plugins/Experimental/GeometryCollectionPlugin/Source/GeometryCollectionNodes/Private/Dataflow/GeometryCollectionMaterialNodes.cpp b/Engine/Plugins/Experimental/GeometryCollectionPlugin/Source/GeometryCollectionNodes/Private/Dataflow/GeometryCollectionMaterialNodes.cpp
--- a/Engine/Plugins/Experimental/GeometryCollectionPlugin/Source/GeometryCollectionNodes/Private/Dataflow/GeometryCollectionMaterialNodes.cpp
+++ b/Engine/Plugins/Experimental/GeometryCollectionPlugin/Source/GeometryCollectionNodes/Private/Dataflow/GeometryCollectionMaterialNodes.cpp
@@ -44,8 +44,6 @@ namespace UE::Dataflow
 {
 	void GeometryCollectionMaterialNodes()
 	{
-		static const FLinearColor CDefaultNodeBodyTintColor = FLinearColor(0.f, 0.f, 0.f, 0.5f);
-
 		DATAFLOW_NODE_REGISTER_CREATION_FACTORY(FAddMaterialToCollectionDataflowNode);
 		DATAFLOW_NODE_REGISTER_CREATION_FACTORY(FReAssignMaterialInCollectionDataflowNode);
 		DATAFLOW_NODE_REGISTER_CREATION_FACTORY(FMaterialsInfoDataflowNode);
@@ -53,9 +51,6 @@ namespace UE::Dataflow
 		DATAFLOW_NODE_REGISTER_CREATION_FACTORY(FSetMaterialInMaterialsArrayDataflowNode);
 		DATAFLOW_NODE_REGISTER_CREATION_FACTORY(FMakeMaterialDataflowNode);
 		DATAFLOW_NODE_REGISTER_CREATION_FACTORY(FMakeMaterialsArrayDataflowNode);
-
-		// Material
-		DATAFLOW_NODE_REGISTER_CREATION_FACTORY_NODE_COLORS_BY_CATEGORY("Generators", FLinearColor(.7f, 0.7f, 0.7f), CDefaultNodeBodyTintColor);
 	}
 }
 
diff --git a/Engine/Plugins/Experimental/GeometryCollectionPlugin/Source/GeometryCollectionNodes/Private/Dataflow/GeometryCollectionMeshNodes.cpp b/Engine/Plugins/Experimental/GeometryCollectionPlugin/Source/GeometryCollectionNodes/Private/Dataflow/GeometryCollectionMeshNodes.cpp
--- a/Engine/Plugins/Experimental/GeometryCollectionPlugin/Source/GeometryCollectionNodes/Private/Dataflow/GeometryCollectionMeshNodes.cpp
+++ b/Engine/Plugins/Experimental/GeometryCollectionPlugin/Source/GeometryCollectionNodes/Private/Dataflow/GeometryCollectionMeshNodes.cpp
@@ -42,8 +42,6 @@ namespace UE::Dataflow
 
 	void GeometryCollectionMeshNodes()
 	{
-		static const FLinearColor CDefaultNodeBodyTintColor = FLinearColor(0.f, 0.f, 0.f, 0.5f);
-
 		DATAFLOW_NODE_REGISTER_CREATION_FACTORY(FPointsToMeshDataflowNode);
 		DATAFLOW_NODE_REGISTER_CREATION_FACTORY(FBoxToMeshDataflowNode);
 		DATAFLOW_NODE_REGISTER_CREATION_FACTORY(FMeshInfoDataflowNode);
@@ -56,10 +54,6 @@ namespace UE::Dataflow
 		DATAFLOW_NODE_REGISTER_CREATION_FACTORY(FGetMeshDataDataflowNode);
 		DATAFLOW_NODE_REGISTER_CREATION_FACTORY(FApplyMeshProcessorToMeshDataflowNode);
 		DATAFLOW_NODE_REGISTER_CREATION_FACTORY(FApplyMeshProcessorToGeometryCollectionDataflowNode);
-
-
-		// Mesh
-		DATAFLOW_NODE_REGISTER_CREATION_FACTORY_NODE_COLORS_BY_CATEGORY("Mesh", FLinearColor(1.f, 0.16f, 0.05f), CDefaultNodeBodyTintColor);
 	}
 }
 
diff --git a/Engine/Plugins/Experimental/GeometryCollectionPlugin/Source/GeometryCollectionNodes/Private/Dataflow/GeometryCollectionNodes.cpp b/Engine/Plugins/Experimental/GeometryCollectionPlugin/Source/GeometryCollectionNodes/Private/Dataflow/GeometryCollectionNodes.cpp
--- a/Engine/Plugins/Experimental/GeometryCollectionPlugin/Source/GeometryCollectionNodes/Private/Dataflow/GeometryCollectionNodes.cpp
+++ b/Engine/Plugins/Experimental/GeometryCollectionPlugin/Source/GeometryCollectionNodes/Private/Dataflow/GeometryCollectionNodes.cpp
@@ -45,8 +45,6 @@ namespace UE::Dataflow
 {
 	void GeometryCollectionEngineNodes()
 	{
-		static const FLinearColor CDefaultNodeBodyTintColor = FLinearColor(0.f, 0.f, 0.f, 0.5f);
-
 		DATAFLOW_NODE_REGISTER_CREATION_FACTORY(FGetCollectionFromAssetDataflowNode);
 		DATAFLOW_NODE_REGISTER_CREATION_FACTORY(FAppendCollectionAssetsDataflowNode);
 		DATAFLOW_NODE_REGISTER_CREATION_FACTORY(FPrintStringDataflowNode);
@@ -90,19 +88,6 @@ namespace UE::Dataflow
 		DATAFLOW_NODE_REGISTER_CREATION_FACTORY(FPointsToCollectionDataflowNode);
 		DATAFLOW_NODE_REGISTER_CREATION_FACTORY(FCollectionToPointsDataflowNode);
 		DATAFLOW_NODE_REGISTER_CREATION_FACTORY(FSpheresToPointsDataflowNode);
-
-		// GeometryCollection
-		DATAFLOW_NODE_REGISTER_CREATION_FACTORY_NODE_COLORS_BY_CATEGORY("GeometryCollection", FLinearColor(0.55f, 0.45f, 1.0f), CDefaultNodeBodyTintColor);
-		// Development
-		DATAFLOW_NODE_REGISTER_CREATION_FACTORY_NODE_COLORS_BY_CATEGORY("Development", FLinearColor(1.f, 0.f, 0.f), CDefaultNodeBodyTintColor);
-		// Utilities|String
-		DATAFLOW_NODE_REGISTER_CREATION_FACTORY_NODE_COLORS_BY_CATEGORY("Utilities|String", FLinearColor(0.5f, 0.f, 0.5f), CDefaultNodeBodyTintColor);
-		// Fracture
-		DATAFLOW_NODE_REGISTER_CREATION_FACTORY_NODE_COLORS_BY_CATEGORY("Fracture", FLinearColor(1.f, 1.f, 0.8f), CDefaultNodeBodyTintColor);
-		// Utilities
-		DATAFLOW_NODE_REGISTER_CREATION_FACTORY_NODE_COLORS_BY_CATEGORY("Utilities", FLinearColor(1.f, 1.f, 0.f), CDefaultNodeBodyTintColor);
-		// Generators
-		DATAFLOW_NODE_REGISTER_CREATION_FACTORY_NODE_COLORS_BY_CATEGORY("Generators", FLinearColor(.4f, 0.8f, 0.f), CDefaultNodeBodyTintColor);
 	}
 }
 
diff --git a/Engine/Plugins/Experimental/GeometryCollectionPlugin/Source/GeometryCollectionNodes/Private/Dataflow/GeometryCollectionNodesPlugin.cpp b/Engine/Plugins/Experimental/GeometryCollectionPlugin/Source/GeometryCollectionNodes/Private/Dataflow/GeometryCollectionNodesPlugin.cpp
--- a/Engine/Plugins/Experimental/GeometryCollectionPlugin/Source/GeometryCollectionNodes/Private/Dataflow/GeometryCollectionNodesPlugin.cpp
+++ b/Engine/Plugins/Experimental/GeometryCollectionPlugin/Source/GeometryCollectionNodes/Private/Dataflow/GeometryCollectionNodesPlugin.cpp
@@ -37,6 +37,7 @@
 #include "Dataflow/SetVertexColorFromFloatArrayNode.h"
 #include "Dataflow/SetVertexColorFromVertexIndicesNode.h"
 #include "Dataflow/SetVertexColorFromVertexSelectionNode.h"
+#include "Dataflow/DataflowNodeAndPinTypeColors.h"
 
 #define LOCTEXT_NAMESPACE "DataflowNodes"
 
@@ -80,6 +81,9 @@ void IGeometryCollectionNodesPlugin::StartupModule()
 PRAGMA_DISABLE_DEPRECATION_WARNINGS
 	DATAFLOW_NODE_REGISTER_CREATION_FACTORY(FSkeletonToCollectionDataflowNode);
 PRAGMA_ENABLE_DEPRECATION_WARNINGS
+
+	// Register Dataflow node and PinType colors
+	UE::Dataflow::RegisterGeometryCollectionNodesColors();
 }
 
 void IGeometryCollectionNodesPlugin::ShutdownModule()
diff --git a/Engine/Plugins/Experimental/GeometryCollectionPlugin/Source/GeometryCollectionNodes/Private/Dataflow/GeometryCollectionOverrideNodes.cpp b/Engine/Plugins/Experimental/GeometryCollectionPlugin/Source/GeometryCollectionNodes/Private/Dataflow/GeometryCollectionOverrideNodes.cpp
--- a/Engine/Plugins/Experimental/GeometryCollectionPlugin/Source/GeometryCollectionNodes/Private/Dataflow/GeometryCollectionOverrideNodes.cpp
+++ b/Engine/Plugins/Experimental/GeometryCollectionPlugin/Source/GeometryCollectionNodes/Private/Dataflow/GeometryCollectionOverrideNodes.cpp
@@ -11,15 +11,10 @@ namespace UE::Dataflow
 {
 	void GeometryCollectionOverrideNodes()
 	{
-		static const FLinearColor CDefaultNodeBodyTintColor = FLinearColor(0.f, 0.f, 0.f, 0.5f);
-
 		DATAFLOW_NODE_REGISTER_CREATION_FACTORY(FGetFloatOverrideFromAssetDataflowNode);
 		DATAFLOW_NODE_REGISTER_CREATION_FACTORY(FGetIntOverrideFromAssetDataflowNode);
 		DATAFLOW_NODE_REGISTER_CREATION_FACTORY(FGetBoolOverrideFromAssetDataflowNode);
 		DATAFLOW_NODE_REGISTER_CREATION_FACTORY(FGetStringOverrideFromAssetDataflowNode);
-
-		// Override
-		DATAFLOW_NODE_REGISTER_CREATION_FACTORY_NODE_COLORS_BY_CATEGORY("Override", FLinearColor(1.f, 0.4f, 0.4f), CDefaultNodeBodyTintColor);
 	}
 }
 
diff --git a/Engine/Plugins/Experimental/GeometryCollectionPlugin/Source/GeometryCollectionNodes/Private/Dataflow/GeometryCollectionSamplingNodes.cpp b/Engine/Plugins/Experimental/GeometryCollectionPlugin/Source/GeometryCollectionNodes/Private/Dataflow/GeometryCollectionSamplingNodes.cpp
--- a/Engine/Plugins/Experimental/GeometryCollectionPlugin/Source/GeometryCollectionNodes/Private/Dataflow/GeometryCollectionSamplingNodes.cpp
+++ b/Engine/Plugins/Experimental/GeometryCollectionPlugin/Source/GeometryCollectionNodes/Private/Dataflow/GeometryCollectionSamplingNodes.cpp
@@ -11,14 +11,9 @@ namespace UE::Dataflow
 
 	void GeometryCollectionSamplingNodes()
 	{
-		static const FLinearColor CDefaultNodeBodyTintColor = FLinearColor(0.f, 0.f, 0.f, 0.5f);
-
 		DATAFLOW_NODE_REGISTER_CREATION_FACTORY(FUniformPointSamplingDataflowNode);
 		DATAFLOW_NODE_REGISTER_CREATION_FACTORY(FNonUniformPointSamplingDataflowNode);
 		DATAFLOW_NODE_REGISTER_CREATION_FACTORY(FVertexWeightedPointSamplingDataflowNode);
-
-		// GeometryCollection|Sampling
-		DATAFLOW_NODE_REGISTER_CREATION_FACTORY_NODE_COLORS_BY_CATEGORY("GeometryCollection|Sampling", FLinearColor(.1f, 1.f, 0.6f), CDefaultNodeBodyTintColor);
 	}
 }
 
diff --git a/Engine/Plugins/Experimental/GeometryCollectionPlugin/Source/GeometryCollectionNodes/Private/Dataflow/GeometryCollectionSelectionNodes.cpp b/Engine/Plugins/Experimental/GeometryCollectionPlugin/Source/GeometryCollectionNodes/Private/Dataflow/GeometryCollectionSelectionNodes.cpp
--- a/Engine/Plugins/Experimental/GeometryCollectionPlugin/Source/GeometryCollectionNodes/Private/Dataflow/GeometryCollectionSelectionNodes.cpp
+++ b/Engine/Plugins/Experimental/GeometryCollectionPlugin/Source/GeometryCollectionNodes/Private/Dataflow/GeometryCollectionSelectionNodes.cpp
@@ -37,8 +37,6 @@ namespace UE::Dataflow
 
 	void GeometryCollectionSelectionNodes()
 	{
-		static const FLinearColor CDefaultNodeBodyTintColor = FLinearColor(0.f, 0.f, 0.f, 0.5f);
-
 		DATAFLOW_NODE_REGISTER_CREATION_FACTORY(FCollectionTransformSelectionAllDataflowNode);
 		DATAFLOW_NODE_REGISTER_CREATION_FACTORY(FCollectionTransformSelectionInfoDataflowNode);
 		DATAFLOW_NODE_REGISTER_CREATION_FACTORY(FCollectionTransformSelectionNoneDataflowNode);
@@ -80,9 +78,6 @@ namespace UE::Dataflow
 		DATAFLOW_NODE_REGISTER_CREATION_FACTORY(FCollectionTransformSelectionInvertDataflowNode);
 		DATAFLOW_NODE_REGISTER_CREATION_FACTORY(FCollectionFaceSelectionInvertDataflowNode);
 		DATAFLOW_NODE_REGISTER_CREATION_FACTORY(FCollectionTransformSelectionSetOperationDataflowNode);
-
-		// GeometryCollection|Selection
-		DATAFLOW_NODE_REGISTER_CREATION_FACTORY_NODE_COLORS_BY_CATEGORY("GeometryCollection|Selection", FLinearColor(1.f, 1.f, 0.05f), CDefaultNodeBodyTintColor);
 	}
 }
 
diff --git a/Engine/Plugins/Experimental/GeometryCollectionPlugin/Source/GeometryCollectionNodes/Private/Dataflow/GeometryCollectionTextureNodes.cpp b/Engine/Plugins/Experimental/GeometryCollectionPlugin/Source/GeometryCollectionNodes/Private/Dataflow/GeometryCollectionTextureNodes.cpp
--- a/Engine/Plugins/Experimental/GeometryCollectionPlugin/Source/GeometryCollectionNodes/Private/Dataflow/GeometryCollectionTextureNodes.cpp
+++ b/Engine/Plugins/Experimental/GeometryCollectionPlugin/Source/GeometryCollectionNodes/Private/Dataflow/GeometryCollectionTextureNodes.cpp
@@ -9,9 +9,6 @@ namespace UE::Dataflow
 	void RegisterGeometryCollectionTextureNodes()
 	{
 		DATAFLOW_NODE_REGISTER_CREATION_FACTORY(FBakeTextureFromCollectionDataflowNode);
-
-		static const FLinearColor CDefaultNodeBodyTintColor = FLinearColor(0.f, 0.f, 0.f, 0.5f);
-		DATAFLOW_NODE_REGISTER_CREATION_FACTORY_NODE_COLORS_BY_CATEGORY("GeometryCollection|Texture", FLinearColor(.1f, 1.f, 0.6f), CDefaultNodeBodyTintColor);
 	}
 }
 
diff --git a/Engine/Plugins/Experimental/GeometryCollectionPlugin/Source/GeometryCollectionNodes/Private/Dataflow/GeometryCollectionUVNodes.cpp b/Engine/Plugins/Experimental/GeometryCollectionPlugin/Source/GeometryCollectionNodes/Private/Dataflow/GeometryCollectionUVNodes.cpp
--- a/Engine/Plugins/Experimental/GeometryCollectionPlugin/Source/GeometryCollectionNodes/Private/Dataflow/GeometryCollectionUVNodes.cpp
+++ b/Engine/Plugins/Experimental/GeometryCollectionPlugin/Source/GeometryCollectionNodes/Private/Dataflow/GeometryCollectionUVNodes.cpp
@@ -12,9 +12,6 @@ namespace UE::Dataflow
 		DATAFLOW_NODE_REGISTER_CREATION_FACTORY(FAddUVChannelDataflowNode);
 		DATAFLOW_NODE_REGISTER_CREATION_FACTORY(FAutoUnwrapUVDataflowNode);
 		DATAFLOW_NODE_REGISTER_CREATION_FACTORY(FBoxProjectUVDataflowNode);
-
-		static const FLinearColor CDefaultNodeBodyTintColor = FLinearColor(0.f, 0.f, 0.f, 0.5f);
-		DATAFLOW_NODE_REGISTER_CREATION_FACTORY_NODE_COLORS_BY_CATEGORY("GeometryCollection|UV", FLinearColor(.1f, 1.f, 0.6f), CDefaultNodeBodyTintColor);
 	}
 }
 
diff --git a/Engine/Plugins/Experimental/GeometryCollectionPlugin/Source/GeometryCollectionNodes/Private/Dataflow/GeometryCollectionUtilityNodes.cpp b/Engine/Plugins/Experimental/GeometryCollectionPlugin/Source/GeometryCollectionNodes/Private/Dataflow/GeometryCollectionUtilityNodes.cpp
--- a/Engine/Plugins/Experimental/GeometryCollectionPlugin/Source/GeometryCollectionNodes/Private/Dataflow/GeometryCollectionUtilityNodes.cpp
+++ b/Engine/Plugins/Experimental/GeometryCollectionPlugin/Source/GeometryCollectionNodes/Private/Dataflow/GeometryCollectionUtilityNodes.cpp
@@ -22,8 +22,6 @@ namespace UE::Dataflow
 
 	void GeometryCollectionUtilityNodes()
 	{
-		static const FLinearColor CDefaultNodeBodyTintColor = FLinearColor(0.f, 0.f, 0.f, 0.5f);
-
 		DATAFLOW_NODE_REGISTER_CREATION_FACTORY(FMakeDataflowConvexDecompositionSettingsNode);
 		DATAFLOW_NODE_REGISTER_CREATION_FACTORY(FCreateLeafConvexHullsDataflowNode);
 		DATAFLOW_NODE_REGISTER_CREATION_FACTORY(FSimplifyConvexHullsDataflowNode);
diff --git a/Engine/Plugins/Experimental/GeometryCollectionPlugin/Source/GeometryCollectionNodes/Public/Dataflow/DataflowNodeAndPinTypeColors.h b/Engine/Plugins/Experimental/GeometryCollectionPlugin/Source/GeometryCollectionNodes/Public/Dataflow/DataflowNodeAndPinTypeColors.h
new file mode 100644
--- /dev/null
+++ b/Engine/Plugins/Experimental/GeometryCollectionPlugin/Source/GeometryCollectionNodes/Public/Dataflow/DataflowNodeAndPinTypeColors.h
@@ -0,0 +1,9 @@
+// Copyright Epic Games, Inc. All Rights Reserved.
+
+#pragma once 
+
+namespace UE::Dataflow
+{
+	void RegisterGeometryCollectionNodesColors();
+}
+
diff --git a/Engine/Plugins/Experimental/GeometryCollectionPlugin/Source/GeometryCollectionNodes/Public/Dataflow/GeometryCollectionFracturingNodes.h b/Engine/Plugins/Experimental/GeometryCollectionPlugin/Source/GeometryCollectionNodes/Public/Dataflow/GeometryCollectionFracturingNodes.h
--- a/Engine/Plugins/Experimental/GeometryCollectionPlugin/Source/GeometryCollectionNodes/Public/Dataflow/GeometryCollectionFracturingNodes.h
+++ b/Engine/Plugins/Experimental/GeometryCollectionPlugin/Source/GeometryCollectionNodes/Public/Dataflow/GeometryCollectionFracturingNodes.h
@@ -1372,6 +1372,14 @@ private:
 	UPROPERTY(EditAnywhere, Category = "Color", meta = (EditCondition = "bApplyColor"));
 	EDataflowVisualizeFractureColoringType ColoringType = EDataflowVisualizeFractureColoringType::ColorByLevel;
 
+	/**  */
+	UPROPERTY(EditAnywhere, Category = "Color", meta = (UIMin = "0", ClampMin = "0", UIMax = "255", ClampMax = "255"))
+	int32 RandomColorRangeMin = 40;
+
+	/**  */
+	UPROPERTY(EditAnywhere, Category = "Color", meta = (UIMin = "0", ClampMin = "0", UIMax = "255", ClampMax = "255"))
+	int32 RandomColorRangeMax = 190;
+
 	/**  */
 	UPROPERTY(EditAnywhere, Category = "Color", meta = (EditCondition = "ColoringType == EDataflowVisualizeFractureColoringType::ColorByAttr", EditConditionHides))
 	FString Attribute;
diff --git a/Engine/Plugins/Experimental/GeometryCollectionPlugin/Source/GeometryCollectionNodes/Public/Dataflow/GeometryCollectionSamplingNodes.h b/Engine/Plugins/Experimental/GeometryCollectionPlugin/Source/GeometryCollectionNodes/Public/Dataflow/GeometryCollectionSamplingNodes.h
--- a/Engine/Plugins/Experimental/GeometryCollectionPlugin/Source/GeometryCollectionNodes/Public/Dataflow/GeometryCollectionSamplingNodes.h
+++ b/Engine/Plugins/Experimental/GeometryCollectionPlugin/Source/GeometryCollectionNodes/Public/Dataflow/GeometryCollectionSamplingNodes.h
@@ -26,6 +26,7 @@ struct FUniformPointSamplingDataflowNode : public FDataflowNode
 {
 	GENERATED_USTRUCT_BODY()
 	DATAFLOW_NODE_DEFINE_INTERNAL(FUniformPointSamplingDataflowNode, "UniformPointSampling", "PointSampling", "")
+	DATAFLOW_NODE_RENDER_TYPE("PointsRender", FName("TArray<FVector>"), "SamplePoints")
 
 public:
 	/** Mesh to sample points on */
@@ -91,6 +92,7 @@ struct FNonUniformPointSamplingDataflowNode : public FDataflowNode
 {
 	GENERATED_USTRUCT_BODY()
 	DATAFLOW_NODE_DEFINE_INTERNAL(FNonUniformPointSamplingDataflowNode, "NonUniformPointSampling", "PointSampling", "")
+	DATAFLOW_NODE_RENDER_TYPE("PointsRender", FName("TArray<FVector>"), "SamplePoints")
 
 public:
 	/** Mesh to sample points on */
@@ -175,6 +177,7 @@ struct FVertexWeightedPointSamplingDataflowNode : public FDataflowNode
 {
 	GENERATED_USTRUCT_BODY()
 	DATAFLOW_NODE_DEFINE_INTERNAL(FVertexWeightedPointSamplingDataflowNode, "VertexWeightedPointSampling", "PointSampling", "")
+	DATAFLOW_NODE_RENDER_TYPE("PointsRender", FName("TArray<FVector>"), "SamplePoints")
 
 public:
 	/** Mesh to sample points on */
diff --git a/Engine/Source/Runtime/Experimental/Dataflow/Core/Private/Dataflow/DataflowNodeColorsRegistry.cpp b/Engine/Source/Runtime/Experimental/Dataflow/Core/Private/Dataflow/DataflowNodeColorsRegistry.cpp
--- a/Engine/Source/Runtime/Experimental/Dataflow/Core/Private/Dataflow/DataflowNodeColorsRegistry.cpp
+++ b/Engine/Source/Runtime/Experimental/Dataflow/Core/Private/Dataflow/DataflowNodeColorsRegistry.cpp
@@ -16,14 +16,7 @@ namespace UE::Dataflow
 
 		for (auto& Elem : NodeColorsMap)
 		{
-			if (ColorsMap.Contains(Elem.Key))
-			{
-				ColorsMap[Elem.Key] = Elem.Value;
-			}
-			else
-			{
-				ColorsMap.Add(Elem.Key, Elem.Value);
-			}
+			ColorsMap.FindOrAdd(Elem.Key) = Elem.Value;
 		}
 	}
 
@@ -122,15 +115,84 @@ namespace UE::Dataflow
 	{
 		for (auto& Elem : NodeColorsMap)
 		{
-			if (ColorsMap.Contains(Elem.Key))
-			{
-				ColorsMap[Elem.Key] = Elem.Value;
-			}
-			else
+			ColorsMap.FindOrAdd(Elem.Key) = Elem.Value;
+		}
+	}
+
+	/* --------------------------------------------------------------------------------------------------------------- */
+
+	FPinColorRegistry::FPinColorRegistry()
+	{
+		UDataflowSettings* DataflowSettings = GetMutableDefault<UDataflowSettings>();
+		DataflowSettingsChangedDelegateHandle = DataflowSettings->GetOnDataflowSettingsChangedPinColorDelegate().AddRaw(this, &FPinColorRegistry::PinColorChangedInSettings);
+
+		const FPinColorMap PinColorMap = DataflowSettings->GetPinColorMap();
+
+		for (auto& Elem : PinColorMap)
+		{
+			ColorMap.FindOrAdd(Elem.Key) = Elem.Value;
+		}
+	}
+
+	FPinColorRegistry::~FPinColorRegistry()
+	{
+		if (UObjectInitialized())
+		{
+			if (IsClassLoaded<UDataflowSettings>())
 			{
-				ColorsMap.Add(Elem.Key, Elem.Value);
+				UDataflowSettings* DataflowSettings = GetMutableDefault<UDataflowSettings>();
+				DataflowSettings->GetOnDataflowSettingsChangedDelegate().Remove(DataflowSettingsChangedDelegateHandle);
 			}
 		}
 	}
 
+	FPinColorRegistry& FPinColorRegistry::Get()
+	{
+		return TLazySingleton<FPinColorRegistry>::Get();
+	}
+
+	void FPinColorRegistry::TearDown()
+	{
+		return TLazySingleton<FPinColorRegistry>::TearDown();
+	}
+
+	void FPinColorRegistry::RegisterPinColor(const FName& PinType, const FLinearColor& PinColor)
+	{
+		if (!ColorMap.Contains(PinType))
+		{
+			ColorMap.Add(PinType, PinColor);
+		}
+
+		// Register colors in DataflowSettings
+		GetMutableDefault<UDataflowSettings>()->RegisterPinColor(PinType, PinColor);
+	}
+
+	FLinearColor FPinColorRegistry::GetPinColor(const FName& PinType)
+	{
+		if (ColorMap.Contains(PinType))
+		{
+			return ColorMap[PinType];
+		}
+
+		return FLinearColor(0.f, 0.f, 0.f);
+	}
+
+	void FPinColorRegistry::PinColorChangedInSettings(const FPinColorMap& PinColorMap)
+	{
+		for (auto& Elem : PinColorMap)
+		{
+			ColorMap.FindOrAdd(Elem.Key) = Elem.Value;
+		}
+	}
+
+	bool FPinColorRegistry::IsPinTypeRegistered(const FName& PinType)
+	{
+		if (ColorMap.Contains(PinType))
+		{
+			return true;
+		}
+
+		return false;
+	}
+
 }
\ No newline at end of file
diff --git a/Engine/Source/Runtime/Experimental/Dataflow/Core/Private/Dataflow/DataflowSettings.cpp b/Engine/Source/Runtime/Experimental/Dataflow/Core/Private/Dataflow/DataflowSettings.cpp
--- a/Engine/Source/Runtime/Experimental/Dataflow/Core/Private/Dataflow/DataflowSettings.cpp
+++ b/Engine/Source/Runtime/Experimental/Dataflow/Core/Private/Dataflow/DataflowSettings.cpp
@@ -4,23 +4,22 @@
 
 #include UE_INLINE_GENERATED_CPP_BY_NAME(DataflowSettings)
 
-namespace UE::Dataflow::Private
-{
-	static const FLinearColor CManagedArrayCollectionPinTypeColor = FLinearColor(0.353393f, 0.454175f, 1.0f, 1.0f);
-	static const FLinearColor CArrayPinTypeColor = FLinearColor(1.0f, 0.172585f, 0.0f, 1.0f);
-	static const FLinearColor CBoxPinTypeColor = FLinearColor(0.013575f, 0.770000f, 0.429609f, 1.0f);
-	static const FLinearColor CSpherePinTypeColor = FLinearColor(0.2f, 0.6f, 1.f, 1.0f);
-	static const FLinearColor CDataflowAnyTypePinTypeColor = FLinearColor(0.3f, 0.3f, 0.3f, 1.0f);
-} // namespace UE::Dataflow::Private
-
 UDataflowSettings::UDataflowSettings(const FObjectInitializer& ObjectInitlaizer)
 	: Super(ObjectInitlaizer)
 {
-	ManagedArrayCollectionPinTypeColor = UE::Dataflow::Private::CManagedArrayCollectionPinTypeColor;
-	ArrayPinTypeColor = UE::Dataflow::Private::CArrayPinTypeColor;
-	BoxPinTypeColor = UE::Dataflow::Private::CBoxPinTypeColor;
-	SpherePinTypeColor = UE::Dataflow::Private::CSpherePinTypeColor;
-	DataflowAnyTypePinTypeColor = UE::Dataflow::Private::CDataflowAnyTypePinTypeColor;
+	TransformLevelColors.LevelColors.AddUninitialized(10);
+	TransformLevelColors.LevelColors[0] = FLinearColor(FColor(0, 255, 255));
+	TransformLevelColors.LevelColors[1] = FLinearColor(FColor(243, 156, 18));
+	TransformLevelColors.LevelColors[2] = FLinearColor(FColor(46, 204, 113));
+	TransformLevelColors.LevelColors[3] = FLinearColor(FColor(255, 255, 0));
+	TransformLevelColors.LevelColors[4] = FLinearColor(FColor(169, 7, 228));
+	TransformLevelColors.LevelColors[5] = FLinearColor(FColor(255, 0, 255));
+	TransformLevelColors.LevelColors[6] = FLinearColor(FColor(26, 188, 156));
+	TransformLevelColors.LevelColors[7] = FLinearColor(FColor(189, 195, 199));
+	TransformLevelColors.LevelColors[8] = FLinearColor(FColor(0, 0, 255));
+	TransformLevelColors.LevelColors[9] = FLinearColor(FColor(0, 255, 0));
+
+	TransformLevelColors.BlankColor = FLinearColor(255, 255, 255, 255);
 }
 
 FName UDataflowSettings::GetCategoryName() const
@@ -40,6 +39,7 @@ void UDataflowSettings::PostEditChangeProperty(FPropertyChangedEvent& PropertyCh
 	if (PropertyChangedEvent.Property != nullptr)
 	{
 		OnDataflowSettingsChangedDelegate.Broadcast(NodeColorsMap);
+		OnDataflowSettingsChangedPinColorDelegate.Broadcast(PinColorMap);
 	}
 
 	Super::PostEditChangeProperty(PropertyChangedEvent);
@@ -56,6 +56,14 @@ FNodeColors UDataflowSettings::RegisterColors(const FName& Category, const FNode
 	return NodeColorsMap[Category];
 }
 
+FLinearColor UDataflowSettings::RegisterPinColor(const FName& InPinType, const FLinearColor& InColor)
+{
+	if (!PinColorMap.Contains(InPinType))
+	{
+		PinColorMap.Add(InPinType, InColor);
+	}
+	return PinColorMap[InPinType];
+}
 
 
 
diff --git a/Engine/Source/Runtime/Experimental/Dataflow/Core/Public/Dataflow/DataflowNode.h b/Engine/Source/Runtime/Experimental/Dataflow/Core/Public/Dataflow/DataflowNode.h
--- a/Engine/Source/Runtime/Experimental/Dataflow/Core/Public/Dataflow/DataflowNode.h
+++ b/Engine/Source/Runtime/Experimental/Dataflow/Core/Public/Dataflow/DataflowNode.h
@@ -772,7 +772,12 @@ private:
 
 #define DATAFLOW_NODE_REGISTER_CREATION_FACTORY_NODE_COLORS_BY_CATEGORY(A, C1, C2)	\
 {																					\
-	::UE::Dataflow::FNodeColorsRegistry::Get().RegisterNodeColors(A, {C1, C2});			\
+	::UE::Dataflow::FNodeColorsRegistry::Get().RegisterNodeColors(A, {C1, C2});		\
+}																					\
+
+#define DATAFLOW_NODE_REGISTER_CREATION_FACTORY_PIN_COLORS_BY_TYPE(A, C)			\
+{																					\
+	::UE::Dataflow::FPinColorRegistry::Get().RegisterPinColor(A, C);				\
 }																					\
 
 }
diff --git a/Engine/Source/Runtime/Experimental/Dataflow/Core/Public/Dataflow/DataflowNodeColorsRegistry.h b/Engine/Source/Runtime/Experimental/Dataflow/Core/Public/Dataflow/DataflowNodeColorsRegistry.h
--- a/Engine/Source/Runtime/Experimental/Dataflow/Core/Public/Dataflow/DataflowNodeColorsRegistry.h
+++ b/Engine/Source/Runtime/Experimental/Dataflow/Core/Public/Dataflow/DataflowNodeColorsRegistry.h
@@ -15,7 +15,7 @@ class FLazySingleton;
 namespace UE::Dataflow
 {
 	//
-	//
+	// Registry for custom Node colors
 	//
 	class FNodeColorsRegistry
 	{
@@ -32,7 +32,31 @@ namespace UE::Dataflow
 		DATAFLOWCORE_API FNodeColorsRegistry();
 		DATAFLOWCORE_API ~FNodeColorsRegistry();
 
-		TMap<FName, FNodeColors > ColorsMap;					// [Category] -> Colors
+		FNodeColorsMap ColorsMap;					// [Category] -> Colors
+		FDelegateHandle DataflowSettingsChangedDelegateHandle;
+
+		friend FLazySingleton;
+	};
+
+	//
+	// Registry for custom Pin colors
+	//
+	class FPinColorRegistry
+	{
+	public:
+		static DATAFLOWCORE_API FPinColorRegistry& Get();
+		static DATAFLOWCORE_API void TearDown();
+
+		DATAFLOWCORE_API void RegisterPinColor(const FName& PinType, const FLinearColor& PinColor);
+		DATAFLOWCORE_API FLinearColor GetPinColor(const FName& PinType);
+		DATAFLOWCORE_API void PinColorChangedInSettings(const FPinColorMap& PinColorMap);
+		DATAFLOWCORE_API bool IsPinTypeRegistered(const FName& PinType);
+
+	private:
+		DATAFLOWCORE_API FPinColorRegistry();
+		DATAFLOWCORE_API ~FPinColorRegistry();
+
+		FPinColorMap ColorMap;					// [PinType] -> Color
 		FDelegateHandle DataflowSettingsChangedDelegateHandle;
 
 		friend FLazySingleton;
diff --git a/Engine/Source/Runtime/Experimental/Dataflow/Core/Public/Dataflow/DataflowSettings.h b/Engine/Source/Runtime/Experimental/Dataflow/Core/Public/Dataflow/DataflowSettings.h
--- a/Engine/Source/Runtime/Experimental/Dataflow/Core/Public/Dataflow/DataflowSettings.h
+++ b/Engine/Source/Runtime/Experimental/Dataflow/Core/Public/Dataflow/DataflowSettings.h
@@ -8,6 +8,7 @@
 #include "DataflowSettings.generated.h"
 
 typedef TMap<FName, FNodeColors> FNodeColorsMap;
+typedef TMap<FName, FLinearColor> FPinColorMap;
 
 USTRUCT()
 struct FNodeColors
@@ -21,40 +22,40 @@ struct FNodeColors
 	FLinearColor NodeBodyTintColor = FLinearColor(0.f, 0.f, 0.f);
 };
 
-UCLASS(config = Editor, defaultconfig, meta = (DisplayName = "Dataflow"), MinimalAPI)
-class UDataflowSettings : public UDeveloperSettings
+USTRUCT()
+struct FTransformLevelColors
 {
-	GENERATED_UCLASS_BODY()
+	GENERATED_USTRUCT_BODY()
 
-	/** TArray<> pin type color. The other pin colors are defined in the general editor settings. */
-	UPROPERTY(config, EditAnywhere, Category = PinColors)
-	FLinearColor ArrayPinTypeColor;
+	UPROPERTY(EditAnywhere, Category = Colors)
+	TArray<FLinearColor> LevelColors;
 
-	/** FManagedArrayCollection pin type color. The other pin colors are defined in the general editor settings. */
-	UPROPERTY(config, EditAnywhere, Category = PinColors)
-	FLinearColor ManagedArrayCollectionPinTypeColor;
+	UPROPERTY(EditAnywhere, Category = Colors)
+	FLinearColor BlankColor;
+};
 
-	/** FBox pin type color. The other pin colors are defined in the general editor settings. */
-	UPROPERTY(config, EditAnywhere, Category = PinColors)
-	FLinearColor BoxPinTypeColor;
+UCLASS(config = Editor, defaultconfig, meta = (DisplayName = "Dataflow"), MinimalAPI)
+class UDataflowSettings : public UDeveloperSettings
+{
+	GENERATED_UCLASS_BODY()
 
-	/** FSphere pin type color. The other pin colors are defined in the general editor settings. */
-	UPROPERTY(config, EditAnywhere, Category = PinColors)
-	FLinearColor SpherePinTypeColor;
+	UPROPERTY(config, EditAnywhere, Category = NodeColors)
+	TMap<FName, FNodeColors> NodeColorsMap;
 
-	/** FDataflowAnyType pin type color. The other pin colors are defined in the general editor settings. */
 	UPROPERTY(config, EditAnywhere, Category = PinColors)
-	FLinearColor DataflowAnyTypePinTypeColor;
+	TMap<FName, FLinearColor> PinColorMap;
 
-	UPROPERTY(config, EditAnywhere, Category = NodeColors)
-	TMap<FName, FNodeColors> NodeColorsMap;
+	UPROPERTY(config, EditAnywhere, Category = TransformLevelColors)
+	FTransformLevelColors TransformLevelColors;
 
 	// Begin UDeveloperSettings Interface
 	DATAFLOWCORE_API virtual FName GetCategoryName() const override;
 
 	DATAFLOWCORE_API FNodeColors RegisterColors(const FName& Category, const FNodeColors& Colors);
+	DATAFLOWCORE_API FLinearColor RegisterPinColor(const FName& PinType, const FLinearColor& Color);
 
 	const TMap<FName, FNodeColors>& GetNodeColorsMap() { return NodeColorsMap; }
+	const TMap<FName, FLinearColor>& GetPinColorMap() { return PinColorMap; }
 
 #if WITH_EDITOR
 	DATAFLOWCORE_API virtual FText GetSectionText() const override;
@@ -67,12 +68,15 @@ class UDataflowSettings : public UDeveloperSettings
 
 public:
 	DECLARE_MULTICAST_DELEGATE_OneParam(FOnDataflowSettingsChanged, const FNodeColorsMap&);
+	DECLARE_MULTICAST_DELEGATE_OneParam(FOnDataflowSettingsChangedPinColor, const FPinColorMap&);
 
 	/** Gets a multicast delegate which is called whenever one of the parameters in this settings object changes. */
 	FOnDataflowSettingsChanged& GetOnDataflowSettingsChangedDelegate() { return OnDataflowSettingsChangedDelegate; }
+	FOnDataflowSettingsChangedPinColor& GetOnDataflowSettingsChangedPinColorDelegate() { return OnDataflowSettingsChangedPinColorDelegate; }
 
 protected:
 	FOnDataflowSettingsChanged OnDataflowSettingsChangedDelegate;
+	FOnDataflowSettingsChangedPinColor OnDataflowSettingsChangedPinColorDelegate;
 
 };
 
diff --git a/Engine/Source/Runtime/Experimental/GeometryCollectionEngine/Private/GeometryCollection/GeometryCollectionComponent.cpp b/Engine/Source/Runtime/Experimental/GeometryCollectionEngine/Private/GeometryCollection/GeometryCollectionComponent.cpp
--- a/Engine/Source/Runtime/Experimental/GeometryCollectionEngine/Private/GeometryCollection/GeometryCollectionComponent.cpp
+++ b/Engine/Source/Runtime/Experimental/GeometryCollectionEngine/Private/GeometryCollection/GeometryCollectionComponent.cpp
@@ -72,6 +72,7 @@
 #include "GeometryCollection/Facades/CollectionRemoveOnBreakFacade.h"
 #include "GeometryCollection/Facades/CollectionInstancedMeshFacade.h"
 #include "GeometryCollection/GeometryCollectionExternalRenderInterface.h"
+#include "Dataflow/DataflowSettings.h"
 
 #include UE_INLINE_GENERATED_CPP_BY_NAME(GeometryCollectionComponent)
 
@@ -5658,9 +5659,9 @@ void FScopedColorEdit::UpdateBoneColors()
 	UGeometryCollection* GeometryCollection = GeometryCollectionEdit.GetRestCollection();
 	if(GeometryCollection)
 	{
-		FGeometryCollection* Collection = GeometryCollection->GetGeometryCollection().Get();
+		const UDataflowSettings* DataflowSettings = GetDefault<UDataflowSettings>();
 
-		FLinearColor BlankColor(FColor(80, 80, 80, 50));
+		FGeometryCollection* Collection = GeometryCollection->GetGeometryCollection().Get();
 
 		const TManagedArray<int>& Parents = Collection->Parent;
 		bool HasLevelAttribute = Collection->HasAttribute("Level", FTransformCollection::TransformGroup);
@@ -5701,7 +5702,7 @@ void FScopedColorEdit::UpdateBoneColors()
 					}
 					else
 					{
-						BoneColor = BlankColor;
+						BoneColor = DataflowSettings->TransformLevelColors.BlankColor;
 					}
 				}
 			}
