Dataflow convex hull to mesh node : new output to an array of dynamic meshes

diff --git a/Engine/Plugins/Experimental/Fracture/Source/FractureEngine/Private/FractureEngineConvex.cpp b/Engine/Plugins/Experimental/Fracture/Source/FractureEngine/Private/FractureEngineConvex.cpp
--- a/Engine/Plugins/Experimental/Fracture/Source/FractureEngine/Private/FractureEngineConvex.cpp
+++ b/Engine/Plugins/Experimental/Fracture/Source/FractureEngine/Private/FractureEngineConvex.cpp
@@ -129,6 +129,62 @@ namespace UE::FractureEngine::Convex
 		return bNoFailures;
 	}
 
+	bool GetConvexHullsAsDynamicMeshes(const FManagedArrayCollection& Collection, TArray<UE::Geometry::FDynamicMesh3>& OutMeshes, bool bRestrictToSelection, const TArrayView<const int32> TransformSelection)
+	{
+		OutMeshes.Reset();
+
+		if (!FGeometryCollectionConvexUtility::HasConvexHullData(&Collection))
+		{
+			// nothing to append
+			return false;
+		}
+
+		const TManagedArray<TSet<int32>>& TransformToConvexInds = Collection.GetAttribute<TSet<int32>>("TransformToConvexIndices", FTransformCollection::TransformGroup);
+		const TManagedArray<Chaos::FConvexPtr>& ConvexHulls = Collection.GetAttribute<Chaos::FConvexPtr>(FGeometryCollection::ConvexHullAttribute, FGeometryCollection::ConvexGroup);
+
+		GeometryCollection::Facades::FCollectionTransformFacade TransformFacade(Collection);
+		TArray<FTransform> GlobalTransformArray = TransformFacade.ComputeCollectionSpaceTransforms();
+
+		auto AppendBone = [&TransformToConvexInds, &ConvexHulls, &GlobalTransformArray, &OutMeshes](int32 BoneIdx) -> bool
+			{
+				if (BoneIdx < 0 || BoneIdx >= TransformToConvexInds.Num())
+				{
+					// invalid bone index
+					return false;
+				}
+				for (int32 ConvexIdx : TransformToConvexInds[BoneIdx])
+				{
+					constexpr bool bConvertNonManifold = true; // Add non-manifold faces so they are still included in the debug visualization
+					constexpr bool bInvertFaces = true; // FConvex mesh data appears to have opposite default winding from what we expect for triangle meshes
+					UE::Geometry::FDynamicMesh3& DynamicMesh = OutMeshes.AddDefaulted_GetRef();
+					AppendConvexHullToCompactDynamicMesh(ConvexHulls[ConvexIdx].GetReference(), DynamicMesh, &GlobalTransformArray[BoneIdx], bConvertNonManifold, bInvertFaces);
+				}
+				return true;
+			};
+
+		bool bNoFailures = true;
+
+		if (bRestrictToSelection)
+		{
+			for (int32 BoneIdx : TransformSelection)
+			{
+				bool bSuccess = AppendBone(BoneIdx);
+				bNoFailures = bNoFailures && bSuccess;
+			}
+		}
+		else
+		{
+			for (int32 BoneIdx = 0; BoneIdx < Collection.NumElements(FGeometryCollection::TransformGroup); ++BoneIdx)
+			{
+				bool bSuccess = AppendBone(BoneIdx);
+				bNoFailures = bNoFailures && bSuccess;
+			}
+		}
+
+		return bNoFailures;
+	}
+
+
 	bool SimplifyConvexHulls(FManagedArrayCollection& Collection, const FSimplifyHullSettings& Settings, bool bRestrictToSelection, const TArrayView<const int32> TransformSelection)
 	{
 		if (!FGeometryCollectionConvexUtility::HasConvexHullData(&Collection))
diff --git a/Engine/Plugins/Experimental/Fracture/Source/FractureEngine/Public/FractureEngineConvex.h b/Engine/Plugins/Experimental/Fracture/Source/FractureEngine/Public/FractureEngineConvex.h
--- a/Engine/Plugins/Experimental/Fracture/Source/FractureEngine/Public/FractureEngineConvex.h
+++ b/Engine/Plugins/Experimental/Fracture/Source/FractureEngine/Public/FractureEngineConvex.h
@@ -43,6 +43,12 @@ namespace UE::FractureEngine::Convex
 	 */
 	bool FRACTUREENGINE_API GetConvexHullsAsDynamicMesh(const FManagedArrayCollection& Collection, UE::Geometry::FDynamicMesh3& UpdateMesh, bool bRestrictToSelection = false, const TArrayView<const int32> TransformSelection = TArrayView<const int32>());
 
+	/**
+	* Get the convex hulls on the given Collection as an array of dynamic meshes. Optionally only include the hulls on the transforms in TransformSelection.
+	* @return true if the collection has convex hulls, the bone selection (if present) was valid, and all requested convex hulls were included in the mesh array 
+	*/
+	bool FRACTUREENGINE_API GetConvexHullsAsDynamicMeshes(const FManagedArrayCollection& Collection, TArray<UE::Geometry::FDynamicMesh3>& UpdateMeshes, bool bRestrictToSelection = false, const TArrayView<const int32> TransformSelection = TArrayView<const int32>());
+
 	/**
 	 * Simplify the convex hulls on the given Collection. Optionally only simplify the hulls on the transforms in TransformSelection.
 	 * @return true if the collection has convex hulls and the hulls were either simplified or did not need to be
diff --git a/Engine/Plugins/Experimental/GeometryCollectionPlugin/Source/GeometryCollectionNodes/Private/Dataflow/GeometryCollectionDebugNodes.cpp b/Engine/Plugins/Experimental/GeometryCollectionPlugin/Source/GeometryCollectionNodes/Private/Dataflow/GeometryCollectionDebugNodes.cpp
--- a/Engine/Plugins/Experimental/GeometryCollectionPlugin/Source/GeometryCollectionNodes/Private/Dataflow/GeometryCollectionDebugNodes.cpp
+++ b/Engine/Plugins/Experimental/GeometryCollectionPlugin/Source/GeometryCollectionNodes/Private/Dataflow/GeometryCollectionDebugNodes.cpp
@@ -26,6 +26,15 @@ namespace UE::Dataflow
 	}
 }
 
+FConvexHullToMeshDataflowNode::FConvexHullToMeshDataflowNode(const UE::Dataflow::FNodeParameters& InParam, FGuid InGuid)
+	: FDataflowNode(InParam, InGuid)
+{
+	RegisterInputConnection(&Collection);
+	RegisterInputConnection(&OptionalSelectionFilter);
+	RegisterOutputConnection(&Mesh);
+	RegisterOutputConnection(&Meshes);
+}
+
 void FConvexHullToMeshDataflowNode::Evaluate(UE::Dataflow::FContext& Context, const FDataflowOutput* Out) const
 {
 	using namespace UE::Geometry;
@@ -49,6 +58,33 @@ void FConvexHullToMeshDataflowNode::Evaluate(UE::Dataflow::FContext& Context, co
 		NewMesh->SetMesh(MoveTemp(HullsMesh));
 		SetValue(Context, NewMesh, &Mesh);
 	}
+	else if (Out->IsA(&Meshes))
+	{
+		const FManagedArrayCollection& InCollection = GetValue<FManagedArrayCollection>(Context, &Collection);
+
+		TArray<int32> TransformSelection;
+		bool bHasSelection = IsConnected(&OptionalSelectionFilter);
+		if (bHasSelection)
+		{
+			const FDataflowTransformSelection& InOptionalSelectionFilter = GetValue(Context, &OptionalSelectionFilter);
+			TransformSelection = InOptionalSelectionFilter.AsArray();
+		}
+
+		TArray<FDynamicMesh3> HullsMeshes;
+		UE::FractureEngine::Convex::GetConvexHullsAsDynamicMeshes(InCollection, HullsMeshes, bHasSelection, TransformSelection);
+
+		TArray<TObjectPtr<UDynamicMesh>> NewMeshes;
+		NewMeshes.SetNum(HullsMeshes.Num());
+
+		for (FDynamicMesh3& HullMesh : HullsMeshes)
+		{
+			TObjectPtr<UDynamicMesh> NewMesh = NewObject<UDynamicMesh>();
+			NewMesh->SetMesh(MoveTemp(HullMesh));
+			NewMeshes.Add(NewMesh);
+		}
+		SetValue(Context, MoveTemp(NewMeshes), &Meshes);
+	}
+
 }
 
 void FSphereCoveringToMeshDataflowNode::Evaluate(UE::Dataflow::FContext& Context, const FDataflowOutput* Out) const
diff --git a/Engine/Plugins/Experimental/GeometryCollectionPlugin/Source/GeometryCollectionNodes/Public/Dataflow/GeometryCollectionDebugNodes.h b/Engine/Plugins/Experimental/GeometryCollectionPlugin/Source/GeometryCollectionNodes/Public/Dataflow/GeometryCollectionDebugNodes.h
--- a/Engine/Plugins/Experimental/GeometryCollectionPlugin/Source/GeometryCollectionNodes/Public/Dataflow/GeometryCollectionDebugNodes.h
+++ b/Engine/Plugins/Experimental/GeometryCollectionPlugin/Source/GeometryCollectionNodes/Public/Dataflow/GeometryCollectionDebugNodes.h
@@ -27,7 +27,7 @@ struct FConvexHullToMeshDataflowNode : public FDataflowNode
 	GENERATED_USTRUCT_BODY()
 	DATAFLOW_NODE_DEFINE_INTERNAL(FConvexHullToMeshDataflowNode, "Convex Hull to Mesh", "GeometryCollection|Mesh", "")
 
-public:
+private:
 	UPROPERTY(meta = (DataflowInput, DataflowIntrinsic))
 	FManagedArrayCollection Collection;
 
@@ -35,18 +35,18 @@ public:
 	UPROPERTY(meta = (DataflowInput, DataflowIntrinsic))
 	FDataflowTransformSelection OptionalSelectionFilter;
 
+	/** Single mesh aggregating all the convex hulls together */
 	UPROPERTY(meta = (DataflowOutput))
 	TObjectPtr<UDynamicMesh> Mesh;
 
-	FConvexHullToMeshDataflowNode(const UE::Dataflow::FNodeParameters& InParam, FGuid InGuid = FGuid::NewGuid())
-		: FDataflowNode(InParam, InGuid)
-	{
-		RegisterInputConnection(&Collection);
-		RegisterInputConnection(&OptionalSelectionFilter);
-		RegisterOutputConnection(&Mesh);
-	}
+	/** Array of meshes for each convex hull found */
+	UPROPERTY(meta = (DataflowOutput))
+	TArray<TObjectPtr<UDynamicMesh>> Meshes;
 
 	virtual void Evaluate(UE::Dataflow::FContext& Context, const FDataflowOutput* Out) const override;
+
+public:
+	FConvexHullToMeshDataflowNode(const UE::Dataflow::FNodeParameters& InParam, FGuid InGuid = FGuid::NewGuid());
 };
 
 // Convert a sphere covering, as generated by the 'protect negative space' option on the "Generate Cluster Convex Hull" nodes, to a dynamic mesh
