more consistently apply guid-gen post-processing in fracture engine's fracture methods + guard against running when the fracture returned INDEX_NONE (i.e., did not fracture)

diff --git a/Engine/Plugins/Experimental/Fracture/Source/FractureEngine/Private/FractureEngineFracturing.cpp b/Engine/Plugins/Experimental/Fracture/Source/FractureEngine/Private/FractureEngineFracturing.cpp
--- a/Engine/Plugins/Experimental/Fracture/Source/FractureEngine/Private/FractureEngineFracturing.cpp
+++ b/Engine/Plugins/Experimental/Fracture/Source/FractureEngine/Private/FractureEngineFracturing.cpp
@@ -14,6 +14,182 @@
 #include "FractureEngineMaterials.h"
 #include "Dataflow/DataflowSettings.h"
 
+// Local helpers
+namespace UE::Private::FractureHelpers
+{
+	static void GenerateVoronoiSites(const FBox& InBoundingBox,
+		const int32 InMinVoronoiSites,
+		const int32 InMaxVoronoiSites,
+		const int32 InRandomSeed,
+		TArray<FVector>& OutSites)
+	{
+		FRandomStream RandStream(InRandomSeed);
+
+		const FVector Extent(InBoundingBox.Max - InBoundingBox.Min);
+
+		const int32 SiteCount = RandStream.RandRange(InMinVoronoiSites, InMaxVoronoiSites);
+
+		OutSites.Reserve(OutSites.Num() + SiteCount);
+		for (int32 ii = 0; ii < SiteCount; ++ii)
+		{
+			OutSites.Emplace(InBoundingBox.Min + FVector(RandStream.FRand(), RandStream.FRand(), RandStream.FRand()) * Extent);
+		}
+	}
+
+	static float GetMaxVertexMovement(const float InGrout,
+		const float InAmplitude,
+		const int32 InOctaveNumber,
+		const float InPersistence)
+	{
+		float MaxDisp = InGrout;
+		float AmplitudeScaled = InAmplitude;
+
+		for (int32 OctaveIdx = 0; OctaveIdx < InOctaveNumber; OctaveIdx++, AmplitudeScaled *= InPersistence)
+		{
+			MaxDisp += FMath::Abs(AmplitudeScaled);
+		}
+
+		return MaxDisp;
+	}
+
+	static FBox GetVoronoiBounds(const FBox& InBoundingBox,
+		const TArray<FVector>& Sites,
+		const float InGrout,
+		const float InAmplitude,
+		const int32 InOctaveNumber,
+		const float InPersistence)
+	{
+		FBox VoronoiBounds = InBoundingBox;
+		if (Sites.Num() > 0)
+		{
+			VoronoiBounds += FBox(Sites);
+		}
+
+		return VoronoiBounds.ExpandBy(GetMaxVertexMovement(InGrout, InAmplitude, InOctaveNumber, InPersistence) + KINDA_SMALL_NUMBER);
+	}
+
+	static void GenerateTemporaryGuids(FManagedArrayCollection& InCollection, int32 InStartIdx, bool InForceInit)
+	{
+		bool bNeedsInit = false;
+		if (!InCollection.HasAttribute("GUID", FTransformCollection::TransformGroup))
+		{
+			FManagedArrayCollection::FConstructionParameters Params(FName(""), false);
+			InCollection.AddAttribute<FGuid>("GUID", FTransformCollection::TransformGroup, Params);
+			bNeedsInit = true;
+		}
+
+		if (bNeedsInit || InForceInit)
+		{
+			TManagedArray<FGuid>& Guids = InCollection.ModifyAttribute<FGuid>("GUID", FTransformCollection::TransformGroup);
+			for (int32 Idx = InStartIdx; Idx < Guids.Num(); ++Idx)
+			{
+				Guids[Idx] = FGuid::NewGuid();
+			}
+		}
+	}
+
+	static void ProcessNewlyFracturedBones(FGeometryCollection& OutGeomCollection, int32 FirstNewGeometryIndex, int32 NewInternalMaterialID = INDEX_NONE)
+	{
+		if (FirstNewGeometryIndex == INDEX_NONE)
+		{
+			return;
+		}
+
+		// Assign internal material
+		if (NewInternalMaterialID > INDEX_NONE)
+		{
+			FFractureEngineMaterials::SetMaterialOnGeometryAfter(OutGeomCollection, FirstNewGeometryIndex, FFractureEngineMaterials::ETargetFaces::InternalFaces, NewInternalMaterialID);
+			OutGeomCollection.ReindexMaterials();
+		}
+
+		// Generate GUIDs
+		GenerateTemporaryGuids(OutGeomCollection, FirstNewGeometryIndex, true);
+	}
+
+	static int32 UniformFractureProc(
+		TUniquePtr<FGeometryCollection>& OutGeomCollection,
+		const TArray<int32>& InTransformSelectionArr,
+		const FUniformFractureProcSettings& InUniformFractureProcSettings)
+	{
+		TArray<FVector> Sites;
+		GenerateVoronoiSites(InUniformFractureProcSettings.BBox,
+			InUniformFractureProcSettings.MinVoronoiSites,
+			InUniformFractureProcSettings.MaxVoronoiSites,
+			InUniformFractureProcSettings.RandomSeed,
+			Sites);
+
+		FBox VoronoiBounds = GetVoronoiBounds(InUniformFractureProcSettings.BBox,
+			Sites,
+			InUniformFractureProcSettings.Grout,
+			InUniformFractureProcSettings.NoiseSettings.Amplitude,
+			InUniformFractureProcSettings.NoiseSettings.Octaves,
+			InUniformFractureProcSettings.NoiseSettings.Persistence);
+
+		FVector Origin = InUniformFractureProcSettings.Transform.GetTranslation();
+		for (FVector& Site : Sites)
+		{
+			Site -= Origin;
+		}
+		VoronoiBounds.Min -= Origin;
+		VoronoiBounds.Max -= Origin;
+		FVoronoiDiagram Voronoi(Sites, VoronoiBounds, .1f);
+
+		FPlanarCells VoronoiPlanarCells = FPlanarCells(Sites, Voronoi);
+		VoronoiPlanarCells.InternalSurfaceMaterials.NoiseSettings = InUniformFractureProcSettings.NoiseSettings;
+
+		int32 FirstNewGeometryIndex = CutMultipleWithPlanarCells(VoronoiPlanarCells, 
+			*OutGeomCollection, 
+			InTransformSelectionArr, 
+			InUniformFractureProcSettings.Grout,
+			InUniformFractureProcSettings.CollisionSampleSpacing,
+			InUniformFractureProcSettings.RandomSeed, 
+			InUniformFractureProcSettings.Transform, 
+			true, 
+			true, 
+			nullptr, 
+			Origin, 
+			InUniformFractureProcSettings.SplitIslands);
+
+		// failed to cut
+		if (FirstNewGeometryIndex == INDEX_NONE)
+		{
+			return INDEX_NONE;
+		}
+
+		ProcessNewlyFracturedBones(*OutGeomCollection, FirstNewGeometryIndex, InUniformFractureProcSettings.InternalMaterialID);
+
+		return FirstNewGeometryIndex;
+	}
+
+	static void SelectLeavesHelper(const FGeometryCollection& GeometryCollection, FDataflowTransformSelection& InOutTransformSelection, int32 BoneIdx)
+	{
+		if (GeometryCollection.SimulationType[BoneIdx] != FGeometryCollection::ESimulationTypes::FST_Rigid)
+		{
+			InOutTransformSelection.SetNotSelected(BoneIdx);
+			for (int32 ChildIdx : GeometryCollection.Children[BoneIdx])
+			{
+				SelectLeavesHelper(GeometryCollection, InOutTransformSelection, ChildIdx);
+			}
+		}
+		else
+		{
+			InOutTransformSelection.SetSelected(BoneIdx);
+		}
+	}
+	
+	static void ConvertToLeafSelection(const FGeometryCollection& GeometryCollection, FDataflowTransformSelection& InOutTransformSelection)
+	{
+		for (int32 BoneIdx = 0; BoneIdx < InOutTransformSelection.Num(); ++BoneIdx)
+		{
+			if (InOutTransformSelection.IsSelected(BoneIdx))
+			{
+				SelectLeavesHelper(GeometryCollection, InOutTransformSelection, BoneIdx);
+			}
+		}
+	}
+}
+
+
 static void AddAdditionalAttributesIfRequired(FManagedArrayCollection& InOutCollection)
 {
 	FGeometryCollectionClusteringUtility::UpdateHierarchyLevelOfChildren(InOutCollection, -1);
@@ -192,36 +368,6 @@
 	}
 }
 
-namespace UE::Private::FractureHelpers
-{
-	static void SelectLeavesHelper(const FGeometryCollection& GeometryCollection, FDataflowTransformSelection& InOutTransformSelection, int32 BoneIdx)
-	{
-		if (GeometryCollection.SimulationType[BoneIdx] != FGeometryCollection::ESimulationTypes::FST_Rigid)
-		{
-			InOutTransformSelection.SetNotSelected(BoneIdx);
-			for (int32 ChildIdx : GeometryCollection.Children[BoneIdx])
-			{
-				SelectLeavesHelper(GeometryCollection, InOutTransformSelection, ChildIdx);
-			}
-		}
-		else
-		{
-			InOutTransformSelection.SetSelected(BoneIdx);
-		}
-	}
-	
-	static void ConvertToLeafSelection(const FGeometryCollection& GeometryCollection, FDataflowTransformSelection& InOutTransformSelection)
-	{
-		for (int32 BoneIdx = 0; BoneIdx < InOutTransformSelection.Num(); ++BoneIdx)
-		{
-			if (InOutTransformSelection.IsSelected(BoneIdx))
-			{
-				SelectLeavesHelper(GeometryCollection, InOutTransformSelection, BoneIdx);
-			}
-		}
-	}
-}
 
-
 int32 FFractureEngineFracturing::VoronoiFracture(FManagedArrayCollection& InOutCollection,
 	FDataflowTransformSelection InTransformSelection,
@@ -312,6 +458,8 @@
 				
 				int ResultGeometryIndex = CutMultipleWithPlanarCells(VoronoiPlanarCells, *GeomCollection, TransformSelectionArr, InGrout, InCollisionSampleSpacing, InRandomSeed, InTransform, true, true, nullptr, Origin, InSplitIslands);
 
+				UE::Private::FractureHelpers::ProcessNewlyFracturedBones(*GeomCollection, ResultGeometryIndex);
+
 				InOutCollection = (const FManagedArrayCollection&)(*GeomCollection);
 
 				return ResultGeometryIndex;
@@ -393,6 +541,8 @@
 
 		int ResultGeometryIndex = CutMultipleWithMultiplePlanes(CuttingPlanes, InternalSurfaceMaterials, *GeomCollection, TransformSelectionArr, GroutVal, CollisionSampleSpacingVal, InRandomSeed, InTransform, true, nullptr, InSplitIslands);
 
+		UE::Private::FractureHelpers::ProcessNewlyFracturedBones(*GeomCollection, ResultGeometryIndex);
+
 		InOutCollection = (const FManagedArrayCollection&)(*GeomCollection);
 
 		return ResultGeometryIndex;
@@ -546,6 +696,8 @@
 
 		int ResultGeometryIndex = CutMultipleWithMultiplePlanes(CuttingPlanes, InternalSurfaceMaterials, *GeomCollection, TransformSelectionArr, GroutVal, CollisionSampleSpacingVal, InRandomSeed, FTransform().Identity, true, nullptr, InSplitIslands);
 
+		UE::Private::FractureHelpers::ProcessNewlyFracturedBones(*GeomCollection, ResultGeometryIndex);
+
 		InOutCollection = (const FManagedArrayCollection&)(*GeomCollection);
 
 		return ResultGeometryIndex;
@@ -922,6 +1074,8 @@
 
 		int32 ResultGeometryIndex = CutMultipleWithPlanarCells(Cells, *GeomCollection, TransformSelectionArr, GroutVal, InPointSpacing, InRandomSeed, InTransform, true, true, nullptr, Origin, InSplitIslands);
 
+		UE::Private::FractureHelpers::ProcessNewlyFracturedBones(*GeomCollection, ResultGeometryIndex);
+
 		InOutCollection = (const FManagedArrayCollection&)(*GeomCollection);
 
 		return ResultGeometryIndex;
@@ -1047,6 +1201,8 @@
 				constexpr bool bSetDefaultInternalMaterialsFromCollection = true;
 				int32 Index = CutWithMesh(CuttingMesh, ScatterTransform, InternalSurfaceMaterials, *GeomCollection, TransformSelectionArr, InCollisionSampleSpacing, FTransform::Identity, bSetDefaultInternalMaterialsFromCollection, nullptr, InSplitIslands);
 
+				UE::Private::FractureHelpers::ProcessNewlyFracturedBones(*GeomCollection, Index);
+
 				int32 NewLen = Algo::RemoveIf(TransformSelectionArr, [&GeomCollection](int32 Bone)
 				{
 					return !GeomCollection->IsVisible(Bone); // remove already-fractured pieces from the to-cut list
@@ -1137,137 +1293,6 @@
 		InRandomSeed, InChanceToFracture, InSplitIslands, InCollisionSampleSpacing);
 }
 
-namespace UE::Dataflow::Private
-{
-	static void GenerateVoronoiSites(const FBox& InBoundingBox,
-		const int32 InMinVoronoiSites,
-		const int32 InMaxVoronoiSites,
-		const int32 InRandomSeed,
-		TArray<FVector>& OutSites)
-	{
-		FRandomStream RandStream(InRandomSeed);
-
-		const FVector Extent(InBoundingBox.Max - InBoundingBox.Min);
-
-		const int32 SiteCount = RandStream.RandRange(InMinVoronoiSites, InMaxVoronoiSites);
-
-		OutSites.Reserve(OutSites.Num() + SiteCount);
-		for (int32 ii = 0; ii < SiteCount; ++ii)
-		{
-			OutSites.Emplace(InBoundingBox.Min + FVector(RandStream.FRand(), RandStream.FRand(), RandStream.FRand()) * Extent);
-		}
-	}
-
-	static float GetMaxVertexMovement(const float InGrout,
-		const float InAmplitude,
-		const int32 InOctaveNumber,
-		const float InPersistence)
-	{
-		float MaxDisp = InGrout;
-		float AmplitudeScaled = InAmplitude;
-
-		for (int32 OctaveIdx = 0; OctaveIdx < InOctaveNumber; OctaveIdx++, AmplitudeScaled *= InPersistence)
-		{
-			MaxDisp += FMath::Abs(AmplitudeScaled);
-		}
-
-		return MaxDisp;
-	}
-
-	static FBox GetVoronoiBounds(const FBox& InBoundingBox,
-		const TArray<FVector>& Sites,
-		const float InGrout,
-		const float InAmplitude,
-		const int32 InOctaveNumber,
-		const float InPersistence)
-	{
-		FBox VoronoiBounds = InBoundingBox;
-		if (Sites.Num() > 0)
-		{
-			VoronoiBounds += FBox(Sites);
-		}
-
-		return VoronoiBounds.ExpandBy(GetMaxVertexMovement(InGrout, InAmplitude, InOctaveNumber, InPersistence) + KINDA_SMALL_NUMBER);
-	}
-
-	static void GenerateTemporaryGuids(FManagedArrayCollection& InCollection, int32 InStartIdx, bool InForceInit)
-	{
-		bool bNeedsInit = false;
-		if (!InCollection.HasAttribute("GUID", FTransformCollection::TransformGroup))
-		{
-			FManagedArrayCollection::FConstructionParameters Params(FName(""), false);
-			InCollection.AddAttribute<FGuid>("GUID", FTransformCollection::TransformGroup, Params);
-			bNeedsInit = true;
-		}
-
-		if (bNeedsInit || InForceInit)
-		{
-			TManagedArray<FGuid>& Guids = InCollection.ModifyAttribute<FGuid>("GUID", FTransformCollection::TransformGroup);
-			for (int32 Idx = InStartIdx; Idx < Guids.Num(); ++Idx)
-			{
-				Guids[Idx] = FGuid::NewGuid();
-			}
-		}
-	}
-
-	static int32 UniformFractureProc(
-		TUniquePtr<FGeometryCollection>& OutGeomCollection,
-		const TArray<int32>& InTransformSelectionArr,
-		const FUniformFractureProcSettings& InUniformFractureProcSettings)
-	{
-		TArray<FVector> Sites;
-		UE::Dataflow::Private::GenerateVoronoiSites(InUniformFractureProcSettings.BBox,
-			InUniformFractureProcSettings.MinVoronoiSites,
-			InUniformFractureProcSettings.MaxVoronoiSites,
-			InUniformFractureProcSettings.RandomSeed,
-			Sites);
-
-		FBox VoronoiBounds = UE::Dataflow::Private::GetVoronoiBounds(InUniformFractureProcSettings.BBox,
-			Sites,
-			InUniformFractureProcSettings.Grout,
-			InUniformFractureProcSettings.NoiseSettings.Amplitude,
-			InUniformFractureProcSettings.NoiseSettings.Octaves,
-			InUniformFractureProcSettings.NoiseSettings.Persistence);
-
-		FVector Origin = InUniformFractureProcSettings.Transform.GetTranslation();
-		for (FVector& Site : Sites)
-		{
-			Site -= Origin;
-		}
-		VoronoiBounds.Min -= Origin;
-		VoronoiBounds.Max -= Origin;
-		FVoronoiDiagram Voronoi(Sites, VoronoiBounds, .1f);
-
-		FPlanarCells VoronoiPlanarCells = FPlanarCells(Sites, Voronoi);
-		VoronoiPlanarCells.InternalSurfaceMaterials.NoiseSettings = InUniformFractureProcSettings.NoiseSettings;
-
-		int32 FirstNewGeometryIndex = CutMultipleWithPlanarCells(VoronoiPlanarCells, 
-			*OutGeomCollection, 
-			InTransformSelectionArr, 
-			InUniformFractureProcSettings.Grout,
-			InUniformFractureProcSettings.CollisionSampleSpacing,
-			InUniformFractureProcSettings.RandomSeed, 
-			InUniformFractureProcSettings.Transform, 
-			true, 
-			true, 
-			nullptr, 
-			Origin, 
-			InUniformFractureProcSettings.SplitIslands);
-
-		// Assign internal material
-		if (InUniformFractureProcSettings.InternalMaterialID > INDEX_NONE)
-		{
-			FFractureEngineMaterials::SetMaterialOnGeometryAfter(*OutGeomCollection, FirstNewGeometryIndex, FFractureEngineMaterials::ETargetFaces::InternalFaces, InUniformFractureProcSettings.InternalMaterialID);
-			OutGeomCollection->ReindexMaterials();
-		}
-
-		// Generate GUIDs
-		GenerateTemporaryGuids(*OutGeomCollection, FirstNewGeometryIndex, true);
-
-		return FirstNewGeometryIndex;
-	}
-}
-
 int32 FFractureEngineFracturing::UniformFracture(
 	FManagedArrayCollection& InOutCollection,
 	FDataflowTransformSelection InTransformSelection, 
@@ -1331,7 +1356,7 @@
 			UniformFractureProcSettings.AddSamplesForCollision = InUniformFractureSettings.AddSamplesForCollision;
 			UniformFractureProcSettings.CollisionSampleSpacing = InUniformFractureSettings.CollisionSampleSpacing;
 
-			int ResultGeometryIndex = UE::Dataflow::Private::UniformFractureProc(
+			int ResultGeometryIndex = UE::Private::FractureHelpers::UniformFractureProc(
 				GeomCollection,
 				TransformSelectionArr,
 				UniformFractureProcSettings);
@@ -1369,7 +1394,7 @@
 					UniformFractureProcSettings.AddSamplesForCollision = InUniformFractureSettings.AddSamplesForCollision;
 					UniformFractureProcSettings.CollisionSampleSpacing = InUniformFractureSettings.CollisionSampleSpacing;
 
-					ResultGeometryIndex = UE::Dataflow::Private::UniformFractureProc(
+					ResultGeometryIndex = UE::Private::FractureHelpers::UniformFractureProc(
 						GeomCollection,
 						TransformSelection,
 						UniformFractureProcSettings);
