add option for dataflow to get 'robust' mesh versions of convex hulls -- by recomputing the hull triangles from the vertices, rather than trying to read the non-manifold approximate polygons from chaos's hull structure back to a valid mesh

diff --git a/Engine/Plugins/Experimental/Fracture/Source/FractureEngine/Private/FractureEngineConvex.cpp b/Engine/Plugins/Experimental/Fracture/Source/FractureEngine/Private/FractureEngineConvex.cpp
--- a/Engine/Plugins/Experimental/Fracture/Source/FractureEngine/Private/FractureEngineConvex.cpp
+++ b/Engine/Plugins/Experimental/Fracture/Source/FractureEngine/Private/FractureEngineConvex.cpp
@@ -2,6 +2,7 @@
 
 #include "FractureEngineConvex.h"
 #include "Chaos/Convex.h"
+#include "CompGeom/ConvexDecomposition3.h"
 #include "DynamicMesh/DynamicMesh3.h"
 #include "DynamicMesh/DynamicMeshAABBTree3.h"
 #include "Spatial/FastWinding.h"
@@ -19,11 +20,11 @@ namespace
 {
 	// Local helpers for converting convex hulls to dynamic meshes, used to run geometry processing tasks on convex hulls (e.g., simplification, computing negative space)
 
-	static void AppendConvexHullToCompactDynamicMesh(const ::Chaos::FConvex* InConvexHull, UE::Geometry::FDynamicMesh3& Mesh, FTransform* OptionalTransform = nullptr, bool bFixNonmanifoldWithDuplicates = false, bool bInvertFaces = false)
+	static void AppendConvexHullToCompactDynamicMesh(const ::Chaos::FConvex* InConvexHull, UE::Geometry::FDynamicMesh3& Mesh, FTransform* OptionalTransform = nullptr, 
+		bool bFixNonmanifoldWithDuplicates = false, bool bInvertFaces = false, bool bRecompute = false)
 	{
 		check(Mesh.IsCompact());
 
-		const ::Chaos::FConvexStructureData& ConvexStructure = InConvexHull->GetStructureData();
 		const int32 NumV = InConvexHull->NumVertices();
 		const int32 NumP = InConvexHull->NumPlanes();
 		int32 StartV = Mesh.MaxVertexID();
@@ -37,6 +38,27 @@ namespace
 			int32 MeshVIdx = Mesh.AppendVertex(V);
 			checkSlow(MeshVIdx == VIdx + StartV); // Must be true because the mesh is compact
 		}
+
+		// Recompute the convex hull from scratch
+		if (bRecompute)
+		{
+			UE::Geometry::FConvexHull3d Hull;
+			Hull.Solve(NumV, [InConvexHull](int32 VID) { return (FVector)InConvexHull->GetVertex(VID); });
+			const TArray<UE::Geometry::FIndex3i>& Tris = Hull.GetTriangles();
+			for (int32 TriIdx = 0; TriIdx < Tris.Num(); ++TriIdx)
+			{
+				UE::Geometry::FIndex3i OffsetTri = Tris[TriIdx].GetOffsetBy(StartV);
+				if (!bInvertFaces) // recomputed hull has faces w/ opposite winding from chaos
+				{
+					Swap(OffsetTri.B, OffsetTri.C);
+				}
+				Mesh.AppendTriangle(OffsetTri);
+			}
+			return;
+		}
+
+		// Not recomputing -- copy the chaos mesh structure out
+		const ::Chaos::FConvexStructureData& ConvexStructure = InConvexHull->GetStructureData();
 		for (int32 PIdx = 0; PIdx < NumP; ++PIdx)
 		{
 			const int32 NumFaceV = ConvexStructure.NumPlaneVertices(PIdx);
@@ -75,7 +97,7 @@ namespace
 
 namespace UE::FractureEngine::Convex
 {
-	bool GetConvexHullsAsDynamicMesh(const FManagedArrayCollection& Collection, UE::Geometry::FDynamicMesh3& OutMesh, bool bRestrictToSelection, const TArrayView<const int32> TransformSelection)
+	bool GetConvexHullsAsDynamicMesh(const FManagedArrayCollection& Collection, UE::Geometry::FDynamicMesh3& OutMesh, bool bRestrictToSelection, const TArrayView<const int32> TransformSelection, bool bRecomputeHulls)
 	{
 		OutMesh.Clear();
 
@@ -91,7 +113,7 @@ namespace UE::FractureEngine::Convex
 		GeometryCollection::Facades::FCollectionTransformFacade TransformFacade(Collection);
 		TArray<FTransform> GlobalTransformArray = TransformFacade.ComputeCollectionSpaceTransforms();
 
-		auto AppendBone = [&TransformToConvexInds, &ConvexHulls, &GlobalTransformArray, &OutMesh](int32 BoneIdx) -> bool
+		auto AppendBone = [&TransformToConvexInds, &ConvexHulls, &GlobalTransformArray, &OutMesh, bRecomputeHulls](int32 BoneIdx) -> bool
 		{
 			if (BoneIdx < 0 || BoneIdx >= TransformToConvexInds.Num())
 			{
@@ -102,7 +124,7 @@ namespace UE::FractureEngine::Convex
 			{
 				constexpr bool bConvertNonManifold = true; // Add non-manifold faces so they are still included in the debug visualization
 				constexpr bool bInvertFaces = true; // FConvex mesh data appears to have opposite default winding from what we expect for triangle meshes
-				AppendConvexHullToCompactDynamicMesh(ConvexHulls[ConvexIdx].GetReference(), OutMesh, &GlobalTransformArray[BoneIdx], bConvertNonManifold, bInvertFaces);
+				AppendConvexHullToCompactDynamicMesh(ConvexHulls[ConvexIdx].GetReference(), OutMesh, &GlobalTransformArray[BoneIdx], bConvertNonManifold, bInvertFaces, bRecomputeHulls);
 			}
 			return true;
 		};
@@ -129,7 +151,7 @@ namespace UE::FractureEngine::Convex
 		return bNoFailures;
 	}
 
-	bool GetConvexHullsAsDynamicMeshes(const FManagedArrayCollection& Collection, TArray<UE::Geometry::FDynamicMesh3>& OutMeshes, bool bRestrictToSelection, const TArrayView<const int32> TransformSelection)
+	bool GetConvexHullsAsDynamicMeshes(const FManagedArrayCollection& Collection, TArray<UE::Geometry::FDynamicMesh3>& OutMeshes, bool bRestrictToSelection, const TArrayView<const int32> TransformSelection, bool bRecomputeHulls)
 	{
 		OutMeshes.Reset();
 
@@ -145,7 +167,7 @@ namespace UE::FractureEngine::Convex
 		GeometryCollection::Facades::FCollectionTransformFacade TransformFacade(Collection);
 		TArray<FTransform> GlobalTransformArray = TransformFacade.ComputeCollectionSpaceTransforms();
 
-		auto AppendBone = [&TransformToConvexInds, &ConvexHulls, &GlobalTransformArray, &OutMeshes](int32 BoneIdx) -> bool
+		auto AppendBone = [&TransformToConvexInds, &ConvexHulls, &GlobalTransformArray, &OutMeshes, bRecomputeHulls](int32 BoneIdx) -> bool
 			{
 				if (BoneIdx < 0 || BoneIdx >= TransformToConvexInds.Num())
 				{
@@ -157,7 +179,7 @@ namespace UE::FractureEngine::Convex
 					constexpr bool bConvertNonManifold = true; // Add non-manifold faces so they are still included in the debug visualization
 					constexpr bool bInvertFaces = true; // FConvex mesh data appears to have opposite default winding from what we expect for triangle meshes
 					UE::Geometry::FDynamicMesh3& DynamicMesh = OutMeshes.AddDefaulted_GetRef();
-					AppendConvexHullToCompactDynamicMesh(ConvexHulls[ConvexIdx].GetReference(), DynamicMesh, &GlobalTransformArray[BoneIdx], bConvertNonManifold, bInvertFaces);
+					AppendConvexHullToCompactDynamicMesh(ConvexHulls[ConvexIdx].GetReference(), DynamicMesh, &GlobalTransformArray[BoneIdx], bConvertNonManifold, bInvertFaces, bRecomputeHulls);
 				}
 				return true;
 			};
diff --git a/Engine/Plugins/Experimental/Fracture/Source/FractureEngine/Public/FractureEngineConvex.h b/Engine/Plugins/Experimental/Fracture/Source/FractureEngine/Public/FractureEngineConvex.h
--- a/Engine/Plugins/Experimental/Fracture/Source/FractureEngine/Public/FractureEngineConvex.h
+++ b/Engine/Plugins/Experimental/Fracture/Source/FractureEngine/Public/FractureEngineConvex.h
@@ -41,13 +41,13 @@ namespace UE::FractureEngine::Convex
 	 * Get the convex hulls on the given Collection as a Dynamic Mesh. Optionally only include the hulls on the transforms in TransformSelection.
 	 * @return true if the collection has convex hulls, the bone selection (if present) was valid, and all requested convex hulls were included in the mesh
 	 */
-	bool FRACTUREENGINE_API GetConvexHullsAsDynamicMesh(const FManagedArrayCollection& Collection, UE::Geometry::FDynamicMesh3& UpdateMesh, bool bRestrictToSelection = false, const TArrayView<const int32> TransformSelection = TArrayView<const int32>());
+	bool FRACTUREENGINE_API GetConvexHullsAsDynamicMesh(const FManagedArrayCollection& Collection, UE::Geometry::FDynamicMesh3& UpdateMesh, bool bRestrictToSelection = false, const TArrayView<const int32> TransformSelection = TArrayView<const int32>(), bool bRecomputeHulls = false);
 
 	/**
 	* Get the convex hulls on the given Collection as an array of dynamic meshes. Optionally only include the hulls on the transforms in TransformSelection.
 	* @return true if the collection has convex hulls, the bone selection (if present) was valid, and all requested convex hulls were included in the mesh array 
 	*/
-	bool FRACTUREENGINE_API GetConvexHullsAsDynamicMeshes(const FManagedArrayCollection& Collection, TArray<UE::Geometry::FDynamicMesh3>& UpdateMeshes, bool bRestrictToSelection = false, const TArrayView<const int32> TransformSelection = TArrayView<const int32>());
+	bool FRACTUREENGINE_API GetConvexHullsAsDynamicMeshes(const FManagedArrayCollection& Collection, TArray<UE::Geometry::FDynamicMesh3>& UpdateMeshes, bool bRestrictToSelection = false, const TArrayView<const int32> TransformSelection = TArrayView<const int32>(), bool bRecomputeHulls = false);
 
 	/**
 	 * Simplify the convex hulls on the given Collection. Optionally only simplify the hulls on the transforms in TransformSelection.
diff --git a/Engine/Plugins/Experimental/GeometryCollectionPlugin/Source/GeometryCollectionNodes/Private/Dataflow/GeometryCollectionDebugNodes.cpp b/Engine/Plugins/Experimental/GeometryCollectionPlugin/Source/GeometryCollectionNodes/Private/Dataflow/GeometryCollectionDebugNodes.cpp
--- a/Engine/Plugins/Experimental/GeometryCollectionPlugin/Source/GeometryCollectionNodes/Private/Dataflow/GeometryCollectionDebugNodes.cpp
+++ b/Engine/Plugins/Experimental/GeometryCollectionPlugin/Source/GeometryCollectionNodes/Private/Dataflow/GeometryCollectionDebugNodes.cpp
@@ -52,7 +52,7 @@ void FConvexHullToMeshDataflowNode::Evaluate(UE::Dataflow::FContext& Context, co
 		}
 
 		FDynamicMesh3 HullsMesh;
-		UE::FractureEngine::Convex::GetConvexHullsAsDynamicMesh(InCollection, HullsMesh, bHasSelection, TransformSelection);
+		UE::FractureEngine::Convex::GetConvexHullsAsDynamicMesh(InCollection, HullsMesh, bHasSelection, TransformSelection, bUseRobustHulls);
 
 		TObjectPtr<UDynamicMesh> NewMesh = NewObject<UDynamicMesh>();
 		NewMesh->SetMesh(MoveTemp(HullsMesh));
@@ -71,7 +71,7 @@ void FConvexHullToMeshDataflowNode::Evaluate(UE::Dataflow::FContext& Context, co
 		}
 
 		TArray<FDynamicMesh3> HullsMeshes;
-		UE::FractureEngine::Convex::GetConvexHullsAsDynamicMeshes(InCollection, HullsMeshes, bHasSelection, TransformSelection);
+		UE::FractureEngine::Convex::GetConvexHullsAsDynamicMeshes(InCollection, HullsMeshes, bHasSelection, TransformSelection, bUseRobustHulls);
 
 		TArray<TObjectPtr<UDynamicMesh>> NewMeshes;
 		NewMeshes.SetNum(HullsMeshes.Num());
diff --git a/Engine/Plugins/Experimental/GeometryCollectionPlugin/Source/GeometryCollectionNodes/Public/Dataflow/GeometryCollectionDebugNodes.h b/Engine/Plugins/Experimental/GeometryCollectionPlugin/Source/GeometryCollectionNodes/Public/Dataflow/GeometryCollectionDebugNodes.h
--- a/Engine/Plugins/Experimental/GeometryCollectionPlugin/Source/GeometryCollectionNodes/Public/Dataflow/GeometryCollectionDebugNodes.h
+++ b/Engine/Plugins/Experimental/GeometryCollectionPlugin/Source/GeometryCollectionNodes/Public/Dataflow/GeometryCollectionDebugNodes.h
@@ -35,6 +35,10 @@ private:
 	UPROPERTY(meta = (DataflowInput))
 	FDataflowTransformSelection OptionalSelectionFilter;
 
+	/** Whether to robustly extract valid/manifold meshes to represent the convex hulls. Note: Not necessary for simple visualization, but useful for downstream processing. */
+	UPROPERTY(EditAnywhere, Category = Options)
+	bool bUseRobustHulls = true;
+
 	/** Single mesh aggregating all the convex hulls together */
 	UPROPERTY(meta = (DataflowOutput))
 	TObjectPtr<UDynamicMesh> Mesh;
diff --git a/Engine/Source/Runtime/GeometryCore/Public/IndexTypes.h b/Engine/Source/Runtime/GeometryCore/Public/IndexTypes.h
--- a/Engine/Source/Runtime/GeometryCore/Public/IndexTypes.h
+++ b/Engine/Source/Runtime/GeometryCore/Public/IndexTypes.h
@@ -220,6 +220,11 @@ PRAGMA_ENABLE_DEPRECATION_WARNINGS
 		return (A == Value) || (B == Value) || (C == Value);
 	}
 
+	/** @return offset triplet, with the OffsetIndicesBy value added to each index */
+	[[nodiscard]] FIndex3i GetOffsetBy(int32 OffsetIndicesBy) const
+	{
+		return FIndex3i(A + OffsetIndicesBy, B + OffsetIndicesBy, C + OffsetIndicesBy);
+	}
 
 	/**
 	 * @return shifted triplet such that A=WantIndex0Value, and B,C values maintain the same relative ordering
