do not in-place update of convex hulls in simplifyconvexhulls method, since cached collections may hold the same pointers

diff --git a/Engine/Plugins/Experimental/Fracture/Source/FractureEngine/Private/FractureEngineConvex.cpp b/Engine/Plugins/Experimental/Fracture/Source/FractureEngine/Private/FractureEngineConvex.cpp
--- a/Engine/Plugins/Experimental/Fracture/Source/FractureEngine/Private/FractureEngineConvex.cpp
+++ b/Engine/Plugins/Experimental/Fracture/Source/FractureEngine/Private/FractureEngineConvex.cpp
@@ -228,7 +228,11 @@ namespace UE::FractureEngine::Convex
 			bool bNoFailures = true;
 			for (int32 ConvexIdx : TransformToConvexInds[BoneIdx])
 			{
-				bool bSuccess = SimplifyConvexHull(ConvexHulls[ConvexIdx].GetReference(), ConvexHulls[ConvexIdx].GetReference(), Settings);
+				// Note: We construct a new convex to hold the simplified result, rather than directly overwriting the old result
+				// to make sure that cached collections holding the same FConvexPtr aren't accidentally updated w/ the simplfied version
+				Chaos::FConvexPtr SimplifiedConvex(new Chaos::FConvex);
+				bool bSuccess = SimplifyConvexHull(ConvexHulls[ConvexIdx].GetReference(), SimplifiedConvex.GetReference(), Settings);
+				ConvexHulls[ConvexIdx] = SimplifiedConvex;
 				bNoFailures = bNoFailures && bSuccess;
 			}
 			return bNoFailures;
