add dataflow selection validation, and validated selection conversion to array, to fix crashes on processing a selection when it does not match the collection

diff --git a/Engine/Plugins/Experimental/Fracture/Source/FractureEngine/Private/FractureEngineFracturing.cpp b/Engine/Plugins/Experimental/Fracture/Source/FractureEngine/Private/FractureEngineFracturing.cpp
--- a/Engine/Plugins/Experimental/Fracture/Source/FractureEngine/Private/FractureEngineFracturing.cpp
+++ b/Engine/Plugins/Experimental/Fracture/Source/FractureEngine/Private/FractureEngineFracturing.cpp
@@ -163,6 +163,10 @@ namespace UE::Private::FractureHelpers
 
 	static void SelectLeavesHelper(const FGeometryCollection& GeometryCollection, FDataflowTransformSelection& InOutTransformSelection, int32 BoneIdx)
 	{
+		if (!ensure(BoneIdx < InOutTransformSelection.Num() && GeometryCollection.SimulationType.IsValidIndex(BoneIdx)))
+		{
+			return;
+		}
 		if (GeometryCollection.SimulationType[BoneIdx] != FGeometryCollection::ESimulationTypes::FST_Rigid)
 		{
 			InOutTransformSelection.SetNotSelected(BoneIdx);
@@ -179,6 +183,11 @@ namespace UE::Private::FractureHelpers
 	
 	static void ConvertToLeafSelection(const FGeometryCollection& GeometryCollection, FDataflowTransformSelection& InOutTransformSelection)
 	{
+		if (!InOutTransformSelection.IsValidForCollection(GeometryCollection))
+		{
+			TArray<int32> ValidSelection = InOutTransformSelection.AsArrayValidated(GeometryCollection);
+			InOutTransformSelection.InitFromArray(GeometryCollection, ValidSelection);
+		}
 		for (int32 BoneIdx = 0; BoneIdx < InOutTransformSelection.Num(); ++BoneIdx)
 		{
 			if (InOutTransformSelection.IsSelected(BoneIdx))
@@ -448,8 +457,9 @@ int32 FFractureEngineFracturing::VoronoiFracture(FManagedArrayCollection& InOutC
 				VoronoiPlanarCells.InternalSurfaceMaterials.NoiseSettings = NoiseSettings;
 
 				RandomReduceSelection(InTransformSelection, InRandomSeed, InChanceToFracture);
+				
 				UE::Private::FractureHelpers::ConvertToLeafSelection(*GeomCollection, InTransformSelection);
-				TArray<int32> TransformSelectionArr = InTransformSelection.AsArray();
+				TArray<int32> TransformSelectionArr = InTransformSelection.AsArrayValidated(*GeomCollection);
 
 				if (!FFractureEngineSelection::IsBoneSelectionValid(InOutCollection, TransformSelectionArr))
 				{
@@ -540,7 +550,7 @@ int32 FFractureEngineFracturing::PlaneCutter(FManagedArrayCollection& InOutColle
 
 		RandomReduceSelection(InTransformSelection, InRandomSeed, InChanceToFracture);
 		UE::Private::FractureHelpers::ConvertToLeafSelection(*GeomCollection, InTransformSelection);
-		TArray<int32> TransformSelectionArr = InTransformSelection.AsArray();
+		TArray<int32> TransformSelectionArr = InTransformSelection.AsArrayValidated(*GeomCollection);
 
 		if (!FFractureEngineSelection::IsBoneSelectionValid(InOutCollection, TransformSelectionArr))
 		{
@@ -692,7 +702,7 @@ int32 FFractureEngineFracturing::SliceCutter(FManagedArrayCollection& InOutColle
 
 		RandomReduceSelection(InTransformSelection, InRandomSeed, InChanceToFracture);
 		UE::Private::FractureHelpers::ConvertToLeafSelection(*GeomCollection, InTransformSelection);
-		TArray<int32> TransformSelectionArr = InTransformSelection.AsArray();
+		TArray<int32> TransformSelectionArr = InTransformSelection.AsArrayValidated(*GeomCollection);
 
 		if (!FFractureEngineSelection::IsBoneSelectionValid(InOutCollection, TransformSelectionArr))
 		{
@@ -1069,7 +1079,7 @@ int32 FFractureEngineFracturing::BrickCutter(FManagedArrayCollection& InOutColle
 
 		RandomReduceSelection(InTransformSelection, InRandomSeed, InChanceToFracture);
 		UE::Private::FractureHelpers::ConvertToLeafSelection(*GeomCollection, InTransformSelection);
-		TArray<int32> TransformSelectionArr = InTransformSelection.AsArray();
+		TArray<int32> TransformSelectionArr = InTransformSelection.AsArrayValidated(*GeomCollection);
 
 		if (!FFractureEngineSelection::IsBoneSelectionValid(InOutCollection, TransformSelectionArr))
 		{
@@ -1197,7 +1207,7 @@ int32 FFractureEngineFracturing::MeshArrayCutter(TArray<FTransform>& MeshTransfo
 		FDataflowTransformSelection InTransformSelection = InTransformSelectionConst;
 		RandomReduceSelection(InTransformSelection, InRandomSeed, InChanceToFracture);
 		UE::Private::FractureHelpers::ConvertToLeafSelection(*GeomCollection, InTransformSelection);
-		TArray<int32> TransformSelectionArr = InTransformSelection.AsArray();
+		TArray<int32> TransformSelectionArr = InTransformSelection.AsArrayValidated(*GeomCollection);
 
 		int32 SequentialMeshIndex = 0;
 		for (const FTransform& ScatterTransform : MeshTransforms)
@@ -1310,7 +1320,7 @@ int32 FFractureEngineFracturing::UniformFracture(
 	{
 		RandomReduceSelection(InTransformSelection, InUniformFractureSettings.RandomSeed, InUniformFractureSettings.ChanceToFracture);
 		UE::Private::FractureHelpers::ConvertToLeafSelection(*GeomCollection, InTransformSelection);
-		TArray<int32> TransformSelectionArr = InTransformSelection.AsArray();
+		TArray<int32> TransformSelectionArr = InTransformSelection.AsArrayValidated(*GeomCollection);
 
 		if (!FFractureEngineSelection::IsBoneSelectionValid(InOutCollection, TransformSelectionArr))
 		{
diff --git a/Engine/Plugins/Experimental/GeometryCollectionPlugin/Source/GeometryCollectionNodes/Private/Dataflow/GeometryCollectionNodes.cpp b/Engine/Plugins/Experimental/GeometryCollectionPlugin/Source/GeometryCollectionNodes/Private/Dataflow/GeometryCollectionNodes.cpp
--- a/Engine/Plugins/Experimental/GeometryCollectionPlugin/Source/GeometryCollectionNodes/Private/Dataflow/GeometryCollectionNodes.cpp
+++ b/Engine/Plugins/Experimental/GeometryCollectionPlugin/Source/GeometryCollectionNodes/Private/Dataflow/GeometryCollectionNodes.cpp
@@ -756,7 +756,7 @@ void FRemoveOnBreakDataflowNode::Evaluate(UE::Dataflow::FContext& Context, const
 		{
 			const FDataflowTransformSelection& InTransformSelection = GetValue(Context, &TransformSelection);
 			TArray<int32> TransformIndices;
-			InTransformSelection.AsArray(TransformIndices);
+			InTransformSelection.AsArrayValidated(TransformIndices, InCollection);
 			RemoveOnBreakFacade.SetFromIndexArray(TransformIndices, Data);
 		}
 		else
@@ -796,13 +796,13 @@ void FSetAnchorStateDataflowNode::Evaluate(UE::Dataflow::FContext& Context, cons
 
 			const bool bAnchored = (AnchorState == EAnchorStateEnum::Dataflow_AnchorState_Anchored) ? true : false;
 			TArray<int32> BoneIndices;
-			InTransformSelection.AsArray(BoneIndices);
+			InTransformSelection.AsArrayValidated(BoneIndices, OutCollection);
 			AnchoringFacade.SetAnchored(BoneIndices, bAnchored);
 
 			if (bSetNotSelectedBonesToOppositeState)
 			{
 				InTransformSelection.Invert();
-				InTransformSelection.AsArray(BoneIndices);
+				InTransformSelection.AsArrayValidated(BoneIndices, OutCollection);
 				AnchoringFacade.SetAnchored(BoneIndices, !bAnchored);
 			}
 		}
@@ -829,7 +829,7 @@ void FSetDynamicStateDataflowNode::Evaluate(UE::Dataflow::FContext& Context, con
 		{
 			const FDataflowTransformSelection& InTransformSelection = GetValue(Context, &TransformSelection);
 
-			const TArray<int32> BoneIndices = InTransformSelection.AsArray();
+			const TArray<int32> BoneIndices = InTransformSelection.AsArrayValidated(OutCollection);
 
 			Chaos::Facades::FCollectionAnchoringFacade AnchoringFacade(OutCollection);
 			if (AnchoringFacade.HasInitialDynamicStateAttribute())
diff --git a/Engine/Source/Runtime/Experimental/Dataflow/Core/Private/Dataflow/DataflowSelection.cpp b/Engine/Source/Runtime/Experimental/Dataflow/Core/Private/Dataflow/DataflowSelection.cpp
--- a/Engine/Source/Runtime/Experimental/Dataflow/Core/Private/Dataflow/DataflowSelection.cpp
+++ b/Engine/Source/Runtime/Experimental/Dataflow/Core/Private/Dataflow/DataflowSelection.cpp
@@ -21,6 +21,34 @@ void FDataflowSelection::Clear()
 	SelectionArray.Init(false, SelectionArray.Num());
 }
 
+void FDataflowSelection::AsArrayValidated(TArray<int32>& OutSelectionArr, const FManagedArrayCollection& InCollection) const
+{
+	if (IsValidForCollection(InCollection))
+	{
+		AsArray(OutSelectionArr);
+	}
+	else
+	{
+		int32 NumEl = InCollection.NumElements(GroupName);
+		UE_LOG(LogChaosDataflow, Warning, TEXT("Selection had mismatched element count vs array: %d vs %d"), Num(), NumEl);
+		int32 ValidNum = FMath::Min(NumEl, Num());
+		for (int32 Idx = 0; Idx < ValidNum; ++Idx)
+		{
+			if (SelectionArray[Idx])
+			{
+				OutSelectionArr.Add(Idx);
+			}
+		}
+	}
+}
+
+TArray<int32> FDataflowSelection::AsArrayValidated(const FManagedArrayCollection& InCollection) const
+{
+	TArray<int32> SelectionArr;
+	AsArrayValidated(SelectionArr, InCollection);
+	return SelectionArr;
+}
+
 void FDataflowSelection::AsArray(TArray<int32>& SelectionArr) const
 {
 	SelectionArr.Reset(SelectionArray.Num());
@@ -150,6 +178,11 @@ void FDataflowSelection::SetSelected(TArray<int32> Indices)
 	}
 }
 
+bool FDataflowSelection::IsValidForCollection(const FManagedArrayCollection& InCollection) const
+{
+	return Num() == InCollection.NumElements(GroupName);
+}
+
 void FDataflowSelection::InitFromArray(const FManagedArrayCollection& InCollection, const TArray<int32>& InSelectionArr)
 {
 	Initialize(InCollection.NumElements(GroupName), false);
diff --git a/Engine/Source/Runtime/Experimental/Dataflow/Core/Public/Dataflow/DataflowSelection.h b/Engine/Source/Runtime/Experimental/Dataflow/Core/Public/Dataflow/DataflowSelection.h
--- a/Engine/Source/Runtime/Experimental/Dataflow/Core/Public/Dataflow/DataflowSelection.h
+++ b/Engine/Source/Runtime/Experimental/Dataflow/Core/Public/Dataflow/DataflowSelection.h
@@ -32,6 +32,8 @@ struct FDataflowSelection
 	void SetNotSelected(int32 Idx) { SelectionArray[Idx] = false; }
 	DATAFLOWCORE_API void AsArray(TArray<int32>& SelectionArr) const;
 	DATAFLOWCORE_API TArray<int32> AsArray() const;
+	DATAFLOWCORE_API void AsArrayValidated(TArray<int32>& SelectionArr, const FManagedArrayCollection& InCollection) const;
+	DATAFLOWCORE_API TArray<int32> AsArrayValidated(const FManagedArrayCollection& InCollection) const;
 	
 	/**
 	* Sets the selection from a sparse array (it only contains the indices of the selected items)
@@ -57,6 +59,9 @@ struct FDataflowSelection
 	*/
 	DATAFLOWCORE_API void InitFromArray(const FManagedArrayCollection& InCollection, const TArray<int32>& InSelectionArr);
 
+	// @return true if this selection is valid for the given collection -- i.e., if the selection expects the group to have the correct number of elements
+	DATAFLOWCORE_API bool IsValidForCollection(const FManagedArrayCollection& InCollection) const;
+
 	/**
 	* Print selection in 
 	* "Selected Transforms: 23 of 34" format;
