Geometry collection - fix crash when merging a single root bone that has an single embedded geometry as a  child

#Synced-CL 50024027

diff --git a/Engine/Plugins/Experimental/Fracture/Source/FractureEngine/Private/FractureEngineClustering.cpp b/Engine/Plugins/Experimental/Fracture/Source/FractureEngine/Private/FractureEngineClustering.cpp
--- a/Engine/Plugins/Experimental/Fracture/Source/FractureEngine/Private/FractureEngineClustering.cpp
+++ b/Engine/Plugins/Experimental/Fracture/Source/FractureEngine/Private/FractureEngineClustering.cpp
@@ -1150,7 +1150,7 @@ bool FFractureEngineClustering::MergeSelectedClusters(FGeometryCollection& Geome
 	{
 		// Cluster selected bones beneath common parent
 		int32 LowestCommonAncestor = FGeometryCollectionClusteringUtility::FindLowestCommonAncestor(&GeometryCollection, Selection);
-		if (LowestCommonAncestor != INDEX_NONE)
+		if (LowestCommonAncestor != INDEX_NONE && HierarchyFacade.GetParent(LowestCommonAncestor) != INDEX_NONE)
 		{
 			FGeometryCollectionClusteringUtility::ClusterBonesUnderNewNodeWithParent(&GeometryCollection, LowestCommonAncestor, Selection, true);
 			GeometryCollection::GenerateTemporaryGuids(&GeometryCollection, StartTransformCount, true);
