Add new split island controls to dataflow fracture and mesh-to-collection nodes

#Synced-CL 52268326

diff --git a/Engine/Plugins/Experimental/Fracture/Source/FractureEngine/Private/FractureEngineFracturing.cpp b/Engine/Plugins/Experimental/Fracture/Source/FractureEngine/Private/FractureEngineFracturing.cpp
--- a/Engine/Plugins/Experimental/Fracture/Source/FractureEngine/Private/FractureEngineFracturing.cpp
+++ b/Engine/Plugins/Experimental/Fracture/Source/FractureEngine/Private/FractureEngineFracturing.cpp
@@ -150,7 +150,7 @@ namespace UE::Private::FractureHelpers
 			true, 
 			nullptr, 
 			Origin, 
-			InUniformFractureProcSettings.SplitIslands);
+			FIslandSplitSettings(InUniformFractureProcSettings.SplitIslands, InUniformFractureProcSettings.CloseVertexDistance, InUniformFractureProcSettings.VertexToSurfaceBridgeDistance));
 
 		// failed to cut
 		if (FirstNewGeometryIndex == INDEX_NONE)
@@ -386,7 +386,7 @@ int32 FFractureEngineFracturing::VoronoiFracture(FManagedArrayCollection& InOutC
 	const FTransform& InTransform,
 	int32 InRandomSeed,
 	float InChanceToFracture,
-	bool InSplitIslands,
+	FIslandSplitSettings IslandSplitSettings,
 	float InGrout,
 	float InAmplitude,
 	float InFrequency,
@@ -468,7 +468,7 @@ int32 FFractureEngineFracturing::VoronoiFracture(FManagedArrayCollection& InOutC
 					return INDEX_NONE;
 				}
 				
-				int ResultGeometryIndex = CutMultipleWithPlanarCells(VoronoiPlanarCells, *GeomCollection, TransformSelectionArr, InGrout, InCollisionSampleSpacing, InRandomSeed, InTransform, true, true, nullptr, Origin, InSplitIslands);
+				int ResultGeometryIndex = CutMultipleWithPlanarCells(VoronoiPlanarCells, *GeomCollection, TransformSelectionArr, InGrout, InCollisionSampleSpacing, InRandomSeed, InTransform, true, true, nullptr, Origin, IslandSplitSettings);
 
 				UE::Private::FractureHelpers::ProcessNewlyFracturedBones(*GeomCollection, ResultGeometryIndex);
 
@@ -508,7 +508,7 @@ int32 FFractureEngineFracturing::PlaneCutter(FManagedArrayCollection& InOutColle
 	int32 InNumPlanes,
 	int32 InRandomSeed,
 	float InChanceToFracture,
-	bool InSplitIslands,
+	FIslandSplitSettings IslandSplitSettings,
 	float InGrout,
 	float InAmplitude,
 	float InFrequency,
@@ -560,7 +560,7 @@ int32 FFractureEngineFracturing::PlaneCutter(FManagedArrayCollection& InOutColle
 			return INDEX_NONE;
 		}
 
-		int ResultGeometryIndex = CutMultipleWithMultiplePlanes(CuttingPlanes, InternalSurfaceMaterials, *GeomCollection, TransformSelectionArr, GroutVal, CollisionSampleSpacingVal, InRandomSeed, InTransform, true, nullptr, InSplitIslands);
+		int ResultGeometryIndex = CutMultipleWithMultiplePlanes(CuttingPlanes, InternalSurfaceMaterials, *GeomCollection, TransformSelectionArr, GroutVal, CollisionSampleSpacingVal, InRandomSeed, InTransform, true, nullptr, IslandSplitSettings);
 
 		UE::Private::FractureHelpers::ProcessNewlyFracturedBones(*GeomCollection, ResultGeometryIndex);
 
@@ -654,7 +654,7 @@ int32 FFractureEngineFracturing::SliceCutter(FManagedArrayCollection& InOutColle
 	float InSliceOffsetVariation,
 	int32 InRandomSeed,
 	float InChanceToFracture,
-	bool InSplitIslands,
+	FIslandSplitSettings IslandSplitSettings,
 	float InGrout,
 	float InAmplitude,
 	float InFrequency,
@@ -715,7 +715,7 @@ int32 FFractureEngineFracturing::SliceCutter(FManagedArrayCollection& InOutColle
 		// Proximity is invalidated.
 		ClearProximity(GeomCollection.Get());
 
-		int ResultGeometryIndex = CutMultipleWithMultiplePlanes(CuttingPlanes, InternalSurfaceMaterials, *GeomCollection, TransformSelectionArr, GroutVal, CollisionSampleSpacingVal, InRandomSeed, FTransform::Identity, true, nullptr, InSplitIslands);
+		int ResultGeometryIndex = CutMultipleWithMultiplePlanes(CuttingPlanes, InternalSurfaceMaterials, *GeomCollection, TransformSelectionArr, GroutVal, CollisionSampleSpacingVal, InRandomSeed, FTransform::Identity, true, nullptr, IslandSplitSettings);
 
 		UE::Private::FractureHelpers::ProcessNewlyFracturedBones(*GeomCollection, ResultGeometryIndex);
 
@@ -1020,7 +1020,7 @@ int32 FFractureEngineFracturing::BrickCutter(FManagedArrayCollection& InOutColle
 	float InBrickDepth,
 	int32 InRandomSeed,
 	float InChanceToFracture,
-	bool InSplitIslands,
+	FIslandSplitSettings IslandSplitSettings,
 	float InGrout,
 	float InAmplitude,
 	float InFrequency,
@@ -1102,7 +1102,7 @@ int32 FFractureEngineFracturing::BrickCutter(FManagedArrayCollection& InOutColle
 		const bool bBricksAreTouching = InGrout <= UE_KINDA_SMALL_NUMBER;
 		FPlanarCells Cells = FPlanarCells(BricksToCut, bBricksAreTouching);
 
-		int32 ResultGeometryIndex = CutMultipleWithPlanarCells(Cells, *GeomCollection, TransformSelectionArr, GroutVal, InPointSpacing, InRandomSeed, InTransform, true, true, nullptr, Origin, InSplitIslands);
+		int32 ResultGeometryIndex = CutMultipleWithPlanarCells(Cells, *GeomCollection, TransformSelectionArr, GroutVal, InPointSpacing, InRandomSeed, InTransform, true, true, nullptr, Origin, IslandSplitSettings);
 
 		UE::Private::FractureHelpers::ProcessNewlyFracturedBones(*GeomCollection, ResultGeometryIndex);
 
@@ -1195,7 +1195,7 @@ int32 FFractureEngineFracturing::MeshArrayCutter(TArray<FTransform>& MeshTransfo
 	const EMeshCutterPerCutMeshSelection PerCutMeshSelection,
 	const int32 InRandomSeed,
 	const float InChanceToFracture,
-	const bool InSplitIslands,
+	const FIslandSplitSettings IslandSplitSettings,
 	const float InCollisionSampleSpacing)
 {
 
@@ -1225,11 +1225,11 @@ int32 FFractureEngineFracturing::MeshArrayCutter(TArray<FTransform>& MeshTransfo
 		for (const FTransform& ScatterTransform : MeshTransforms)
 		{
 			auto ApplyCut = 
-				[&ScatterTransform, &InternalSurfaceMaterials, &GeomCollection, &TransformSelectionArr, InCollisionSampleSpacing, InSplitIslands, &ResultGeometryIndex]
+				[&ScatterTransform, &InternalSurfaceMaterials, &GeomCollection, &TransformSelectionArr, InCollisionSampleSpacing, IslandSplitSettings, &ResultGeometryIndex]
 				(const UE::Geometry::FDynamicMesh3& CuttingMesh)
 			{
 				constexpr bool bSetDefaultInternalMaterialsFromCollection = true;
-				int32 Index = CutWithMesh(CuttingMesh, ScatterTransform, InternalSurfaceMaterials, *GeomCollection, TransformSelectionArr, InCollisionSampleSpacing, FTransform::Identity, bSetDefaultInternalMaterialsFromCollection, nullptr, InSplitIslands);
+				int32 Index = CutWithMesh(CuttingMesh, ScatterTransform, InternalSurfaceMaterials, *GeomCollection, TransformSelectionArr, InCollisionSampleSpacing, FTransform::Identity, bSetDefaultInternalMaterialsFromCollection, nullptr, IslandSplitSettings);
 
 				UE::Private::FractureHelpers::ProcessNewlyFracturedBones(*GeomCollection, Index);
 
@@ -1314,13 +1314,13 @@ int32 FFractureEngineFracturing::MeshCutter(TArray<FTransform>& MeshTransforms,
 	const UE::Geometry::FDynamicMesh3& InDynCuttingMesh,
 	const int32 InRandomSeed,
 	const float InChanceToFracture,
-	const bool InSplitIslands,
+	const FIslandSplitSettings IslandSplitSettings,
 	const float InCollisionSampleSpacing)
 {
 	const UE::Geometry::FDynamicMesh3* LocalCutterPtr = &InDynCuttingMesh;
 	return MeshArrayCutter(MeshTransforms, InOutCollection, InTransformSelection, 
 		TArrayView<const UE::Geometry::FDynamicMesh3*>(&LocalCutterPtr, 1), EMeshCutterPerCutMeshSelection::All,
-		InRandomSeed, InChanceToFracture, InSplitIslands, InCollisionSampleSpacing);
+		InRandomSeed, InChanceToFracture, IslandSplitSettings, InCollisionSampleSpacing);
 }
 
 int32 FFractureEngineFracturing::UniformFracture(
@@ -1381,6 +1381,8 @@ int32 FFractureEngineFracturing::UniformFracture(
 			UniformFractureProcSettings.InternalMaterialID = InUniformFractureSettings.InternalMaterialID;
 			UniformFractureProcSettings.RandomSeed = InUniformFractureSettings.RandomSeed;
 			UniformFractureProcSettings.SplitIslands = InUniformFractureSettings.SplitIslands;
+			UniformFractureProcSettings.CloseVertexDistance = InUniformFractureSettings.CloseVertexDistance;
+			UniformFractureProcSettings.VertexToSurfaceBridgeDistance = InUniformFractureSettings.VertexToSurfaceBridgeDistance;
 			UniformFractureProcSettings.Grout = InUniformFractureSettings.Grout;
 			UniformFractureProcSettings.NoiseSettings = InUniformFractureSettings.NoiseSettings;
 			UniformFractureProcSettings.AddSamplesForCollision = InUniformFractureSettings.AddSamplesForCollision;
@@ -1419,6 +1421,8 @@ int32 FFractureEngineFracturing::UniformFracture(
 					UniformFractureProcSettings.InternalMaterialID = InUniformFractureSettings.InternalMaterialID;
 					UniformFractureProcSettings.RandomSeed = Seed;
 					UniformFractureProcSettings.SplitIslands = InUniformFractureSettings.SplitIslands;
+					UniformFractureProcSettings.CloseVertexDistance = InUniformFractureSettings.CloseVertexDistance;
+					UniformFractureProcSettings.VertexToSurfaceBridgeDistance = InUniformFractureSettings.VertexToSurfaceBridgeDistance;
 					UniformFractureProcSettings.Grout = InUniformFractureSettings.Grout;
 					UniformFractureProcSettings.NoiseSettings = InUniformFractureSettings.NoiseSettings;
 					UniformFractureProcSettings.AddSamplesForCollision = InUniformFractureSettings.AddSamplesForCollision;
diff --git a/Engine/Plugins/Experimental/Fracture/Source/FractureEngine/Public/FractureEngineFracturing.h b/Engine/Plugins/Experimental/Fracture/Source/FractureEngine/Public/FractureEngineFracturing.h
--- a/Engine/Plugins/Experimental/Fracture/Source/FractureEngine/Public/FractureEngineFracturing.h
+++ b/Engine/Plugins/Experimental/Fracture/Source/FractureEngine/Public/FractureEngineFracturing.h
@@ -57,6 +57,8 @@ struct FUniformFractureSettings
 	float ChanceToFracture;
 	bool GroupFracture;
 	bool SplitIslands;
+	double CloseVertexDistance = 1e-3;
+	double VertexToSurfaceBridgeDistance = 0;
 	float Grout;
 	FNoiseSettings NoiseSettings;
 	bool AddSamplesForCollision;
@@ -72,6 +74,8 @@ struct FUniformFractureProcSettings
 	int32 InternalMaterialID;
 	int32 RandomSeed;
 	bool SplitIslands;
+	double CloseVertexDistance = 1e-3;
+	double VertexToSurfaceBridgeDistance = 0;
 	float Grout;
 	FNoiseSettings NoiseSettings;
 	bool AddSamplesForCollision;
@@ -94,7 +98,7 @@ public:
 		const FTransform& InTransform,
 		int32 InRandomSeed,
 		float InChanceToFracture,
-		bool InSplitIslands,
+		FIslandSplitSettings IslandSplitSettings,
 		float InGrout,
 		float InAmplitude,
 		float InFrequency,
@@ -118,7 +122,7 @@ public:
 		int32 InNumPlanes,
 		int32 InRandomSeed,
 		float InChanceToFracture,
-		bool InSplitIslands,
+		FIslandSplitSettings IslandSplitSettings,
 		float InGrout,
 		float InAmplitude,
 		float InFrequency,
@@ -149,7 +153,7 @@ public:
 		float InSliceOffsetVariation,
 		int32 InRandomSeed,
 		float InChanceToFracture,
-		bool InSplitIslands,
+		FIslandSplitSettings IslandSplitSettings,
 		float InGrout,
 		float InAmplitude,
 		float InFrequency,
@@ -182,7 +186,7 @@ public:
 		float InBrickDepth,
 		int32 InRandomSeed,
 		float InChanceToFracture,
-		bool InSplitIslands,
+		FIslandSplitSettings IslandSplitSettings,
 		float InGrout,
 		float InAmplitude,
 		float InFrequency,
@@ -215,7 +219,7 @@ public:
 		const UE::Geometry::FDynamicMesh3& InDynCuttingMesh,
 		const int32 InRandomSeed,
 		const float InChanceToFracture,
-		const bool InSplitIslands,
+		const FIslandSplitSettings IslandSplitSettings,
 		const float InCollisionSampleSpacing);
 
 	static UE_API int32 MeshArrayCutter(TArray<FTransform>& MeshTransforms,
@@ -225,7 +229,7 @@ public:
 		const EMeshCutterPerCutMeshSelection PerCutMeshSelection,
 		const int32 InRandomSeed,
 		const float InChanceToFracture,
-		const bool InSplitIslands,
+		const FIslandSplitSettings IslandSplitSettings,
 		const float InCollisionSampleSpacing);
 
 	static UE_API int32 UniformFracture(
diff --git a/Engine/Plugins/Experimental/GeometryCollectionPlugin/Source/GeometryCollectionNodes/Private/Dataflow/GeometryCollectionFracturingNodes.cpp b/Engine/Plugins/Experimental/GeometryCollectionPlugin/Source/GeometryCollectionNodes/Private/Dataflow/GeometryCollectionFracturingNodes.cpp
--- a/Engine/Plugins/Experimental/GeometryCollectionPlugin/Source/GeometryCollectionNodes/Private/Dataflow/GeometryCollectionFracturingNodes.cpp
+++ b/Engine/Plugins/Experimental/GeometryCollectionPlugin/Source/GeometryCollectionNodes/Private/Dataflow/GeometryCollectionFracturingNodes.cpp
@@ -529,7 +529,7 @@ void FVoronoiFractureDataflowNode_v2::Evaluate(UE::Dataflow::FContext& Context,
 				GetValue<FTransform>(Context, &Transform),
 				0, // RandomSeed is not used in Voronoi fracture, it is used in the source point generation
 				GetValue<float>(Context, &ChanceToFracture),
-				SplitIslands,
+				FIslandSplitSettings(SplitIslands, (double)CloseVertexDistance, (double)VertexToSurfaceBridgeDistance),
 				GetValue<float>(Context, &Grout),
 				GetValue<float>(Context, &Amplitude),
 				GetValue<float>(Context, &Frequency),
@@ -685,7 +685,7 @@ void FPlaneCutterDataflowNode_v2::Evaluate(UE::Dataflow::FContext& Context, cons
 				InNumPlanes,
 				GetValue(Context, &RandomSeed),
 				GetValue(Context, &ChanceToFracture),
-				SplitIslands,
+				FIslandSplitSettings(SplitIslands, (double)CloseVertexDistance, (double)VertexToSurfaceBridgeDistance),
 				GetValue(Context, &Grout),
 				GetValue(Context, &Amplitude),
 				GetValue(Context, &Frequency),
@@ -910,7 +910,7 @@ void FSliceCutterDataflowNode::Evaluate(UE::Dataflow::FContext& Context, const F
 				GetValue<float>(Context, &SliceOffsetVariation),
 				GetValue<int32>(Context, &RandomSeed),
 				GetValue<float>(Context, &ChanceToFracture),
-				SplitIslands,
+				FIslandSplitSettings(SplitIslands, (double)CloseVertexDistance, (double)VertexToSurfaceBridgeDistance),
 				GetValue<float>(Context, &Grout),
 				GetValue<float>(Context, &Amplitude),
 				GetValue<float>(Context, &Frequency),
@@ -1015,7 +1015,7 @@ void FBrickCutterDataflowNode::Evaluate(UE::Dataflow::FContext& Context, const F
 				GetValue<float>(Context, &BrickDepth),
 				GetValue<int32>(Context, &RandomSeed),
 				GetValue<float>(Context, &ChanceToFracture),
-				SplitIslands,
+				FIslandSplitSettings(SplitIslands, (double)CloseVertexDistance, (double)VertexToSurfaceBridgeDistance),
 				GetValue<float>(Context, &Grout),
 				GetValue<float>(Context, &Amplitude),
 				GetValue<float>(Context, &Frequency),
@@ -1199,7 +1199,7 @@ void FMeshCutterDataflowNode::Evaluate(UE::Dataflow::FContext& Context, const FD
 					InPerCutMeshSelection,
 					InRandomSeed,
 					InChanceToFracture,
-					SplitIslands,
+					FIslandSplitSettings(SplitIslands, (double)CloseVertexDistance, (double)VertexToSurfaceBridgeDistance),
 					InCollisionSampleSpacing);
 
 				FDataflowTransformSelection NewSelection;
@@ -1331,6 +1331,8 @@ void FUniformFractureDataflowNode::Evaluate(UE::Dataflow::FContext& Context, con
 			UniformFractureSettings.ChanceToFracture = GetValue(Context, &ChanceToFracture);
 			UniformFractureSettings.GroupFracture = GroupFracture;
 			UniformFractureSettings.SplitIslands = SplitIslands;
+			UniformFractureSettings.CloseVertexDistance = (double)CloseVertexDistance;
+			UniformFractureSettings.VertexToSurfaceBridgeDistance = (double)VertexToSurfaceBridgeDistance;
 			UniformFractureSettings.Grout = GetValue(Context, &Grout);
 			UniformFractureSettings.NoiseSettings.Amplitude = GetValue(Context, &Amplitude);
 			UniformFractureSettings.NoiseSettings.Frequency = GetValue(Context, &Frequency);
diff --git a/Engine/Plugins/Experimental/GeometryCollectionPlugin/Source/GeometryCollectionNodes/Private/Dataflow/GeometryCollectionMeshNodes.cpp b/Engine/Plugins/Experimental/GeometryCollectionPlugin/Source/GeometryCollectionNodes/Private/Dataflow/GeometryCollectionMeshNodes.cpp
--- a/Engine/Plugins/Experimental/GeometryCollectionPlugin/Source/GeometryCollectionNodes/Private/Dataflow/GeometryCollectionMeshNodes.cpp
+++ b/Engine/Plugins/Experimental/GeometryCollectionPlugin/Source/GeometryCollectionNodes/Private/Dataflow/GeometryCollectionMeshNodes.cpp
@@ -193,6 +193,11 @@ void FMeshToCollectionDataflowNode::Evaluate(UE::Dataflow::FContext& Context, co
 					{
 						Components.ConnectCloseVertices(DynMesh, ConnectVerticesThreshold, 2);
 					}
+					if (VertexToSurfaceBridgeDistance > 0)
+					{
+						TMeshAABBTree3<FDynamicMesh3> Spatial(&DynMesh);
+						Components.ConnectVerticesToNearestDifferentComponent(Spatial, (double)VertexToSurfaceBridgeDistance);
+					}
 					FDynamicMeshEditor::SplitMesh(&DynMesh, SplitMeshes, false, [&Components, &DynMesh](int32 TID)
 						{
 							return Components.GetComponent(DynMesh.GetTriangle(TID).A);
diff --git a/Engine/Plugins/Experimental/GeometryCollectionPlugin/Source/GeometryCollectionNodes/Public/Dataflow/GeometryCollectionFracturingNodes.h b/Engine/Plugins/Experimental/GeometryCollectionPlugin/Source/GeometryCollectionNodes/Public/Dataflow/GeometryCollectionFracturingNodes.h
--- a/Engine/Plugins/Experimental/GeometryCollectionPlugin/Source/GeometryCollectionNodes/Public/Dataflow/GeometryCollectionFracturingNodes.h
+++ b/Engine/Plugins/Experimental/GeometryCollectionPlugin/Source/GeometryCollectionNodes/Public/Dataflow/GeometryCollectionFracturingNodes.h
@@ -528,10 +528,18 @@ public:
 	UPROPERTY(EditAnywhere, Category = "Common Fracture", meta = (DataflowInput, UIMin = "0.0", UIMax = "1.0", ClampMin = "0.0", ClampMax = "1.0"));
 	float ChanceToFracture = 1.f;
 
-	/** Whether to split the fractured mesh pieces based on geometric connectivity after fracturing */
-	UPROPERTY(EditAnywhere, Category = "Common Fracture");
+	/** Whether to split the fractured mesh pieces into separate islands based on geometric connectivity after fracturing */
+	UPROPERTY(EditAnywhere, Category = "Common Fracture|Split Islands");
 	bool SplitIslands = true;
 
+	/** Distance threshold for connecting spatially-close vertices before computing islands */
+	UPROPERTY(EditAnywhere, Category = "Common Fracture|Split Islands", meta = (EditCondition = "SplitIslands", Units = "cm", ClampMin = "0.0"));
+	float CloseVertexDistance = 0.001f;
+
+	/** If > 0, bridge separate islands whose surfaces are within this vertex-to-triangle distance. 0 = disabled. */
+	UPROPERTY(EditAnywhere, Category = "Common Fracture|Split Islands", meta = (EditCondition = "SplitIslands", Units = "cm", ClampMin = "0.0"));
+	float VertexToSurfaceBridgeDistance = 0.f;
+
 	/** Amount of space to leave between cut pieces */
 	UPROPERTY(EditAnywhere, Category = "Common Fracture", meta = (DataflowInput, UIMin = 0.f));
 	float Grout = 0.f;
@@ -731,10 +739,18 @@ public:
 	UPROPERTY(EditAnywhere, Category = "Common Fracture", meta = (DataflowInput, UIMin = 0.f, DisplayName = "Chance To Fracture Per Bone", UIMax = 1.f));
 	float ChanceToFracture = 1.f;
 
-	/** Whether to split the fractured mesh pieces based on geometric connectivity after fracturing */
-	UPROPERTY(EditAnywhere, Category = "Common Fracture");
+	/** Whether to split the fractured mesh pieces into separate islands based on geometric connectivity after fracturing */
+	UPROPERTY(EditAnywhere, Category = "Common Fracture|Split Islands");
 	bool SplitIslands = true;
 
+	/** Distance threshold for connecting spatially-close vertices before computing islands */
+	UPROPERTY(EditAnywhere, Category = "Common Fracture|Split Islands", meta = (EditCondition = "SplitIslands", Units = "cm", ClampMin = "0.0"));
+	float CloseVertexDistance = 0.001f;
+
+	/** If > 0, bridge separate islands whose surfaces are within this vertex-to-triangle distance. 0 = disabled. */
+	UPROPERTY(EditAnywhere, Category = "Common Fracture|Split Islands", meta = (EditCondition = "SplitIslands", Units = "cm", ClampMin = "0.0"));
+	float VertexToSurfaceBridgeDistance = 0.f;
+
 	/** Amount of space to leave between cut pieces */
 	UPROPERTY(EditAnywhere, Category = "Common Fracture", meta = (DataflowInput, UIMin = 0.f));
 	float Grout = 0.f;
@@ -943,10 +959,18 @@ public:
 	UPROPERTY(EditAnywhere, Category = "Common Fracture", meta = (DataflowInput, UIMin = 0.f, UIMax = 1.f));
 	float ChanceToFracture = 1.f;
 
-	/** Whether to split the fractured mesh pieces based on geometric connectivity after fracturing */
-	UPROPERTY(EditAnywhere, Category = "Common Fracture");
+	/** Whether to split the fractured mesh pieces into separate islands based on geometric connectivity after fracturing */
+	UPROPERTY(EditAnywhere, Category = "Common Fracture|Split Islands");
 	bool SplitIslands = true;
 
+	/** Distance threshold for connecting spatially-close vertices before computing islands */
+	UPROPERTY(EditAnywhere, Category = "Common Fracture|Split Islands", meta = (EditCondition = "SplitIslands", Units = "cm", ClampMin = "0.0"));
+	float CloseVertexDistance = 0.001f;
+
+	/** If > 0, bridge separate islands whose surfaces are within this vertex-to-triangle distance. 0 = disabled. */
+	UPROPERTY(EditAnywhere, Category = "Common Fracture|Split Islands", meta = (EditCondition = "SplitIslands", Units = "cm", ClampMin = "0.0"));
+	float VertexToSurfaceBridgeDistance = 0.f;
+
 	/** Amount of space to leave between cut pieces */
 	UPROPERTY(EditAnywhere, Category = "Common Fracture", meta = (DataflowInput, UIMin = 0.f));
 	float Grout = 0.f;
@@ -1081,10 +1105,18 @@ public:
 	UPROPERTY(EditAnywhere, Category = "Common Fracture", meta = (DataflowInput, UIMin = 0.f, UIMax = 1.f));
 	float ChanceToFracture = 1.f;
 
-	/** Whether to split the fractured mesh pieces based on geometric connectivity after fracturing */
-	UPROPERTY(EditAnywhere, Category = "Common Fracture");
+	/** Whether to split the fractured mesh pieces into separate islands based on geometric connectivity after fracturing */
+	UPROPERTY(EditAnywhere, Category = "Common Fracture|Split Islands");
 	bool SplitIslands = true;
 
+	/** Distance threshold for connecting spatially-close vertices before computing islands */
+	UPROPERTY(EditAnywhere, Category = "Common Fracture|Split Islands", meta = (EditCondition = "SplitIslands", Units = "cm", ClampMin = "0.0"));
+	float CloseVertexDistance = 0.001f;
+
+	/** If > 0, bridge separate islands whose surfaces are within this vertex-to-triangle distance. 0 = disabled. */
+	UPROPERTY(EditAnywhere, Category = "Common Fracture|Split Islands", meta = (EditCondition = "SplitIslands", Units = "cm", ClampMin = "0.0"));
+	float VertexToSurfaceBridgeDistance = 0.f;
+
 	/** Amount of space to leave between cut pieces */
 	UPROPERTY(EditAnywhere, Category = "Common Fracture", meta = (DataflowInput, UIMin = 0.f));
 	float Grout = 0.f;
@@ -1269,10 +1301,18 @@ public:
 	UPROPERTY(EditAnywhere, Category = "Common Fracture", meta = (DataflowInput, UIMin = 0.f, UIMax = 1.f));
 	float ChanceToFracture = 1.f;
 
-	/** Whether to split the fractured mesh pieces based on geometric connectivity after fracturing */
-	UPROPERTY(EditAnywhere, Category = "Common Fracture");
+	/** Whether to split the fractured mesh pieces into separate islands based on geometric connectivity after fracturing */
+	UPROPERTY(EditAnywhere, Category = "Common Fracture|Split Islands");
 	bool SplitIslands = true;
-	
+
+	/** Distance threshold for connecting spatially-close vertices before computing islands */
+	UPROPERTY(EditAnywhere, Category = "Common Fracture|Split Islands", meta = (EditCondition = "SplitIslands", Units = "cm", ClampMin = "0.0"));
+	float CloseVertexDistance = 0.001f;
+
+	/** If > 0, bridge separate islands whose surfaces are within this vertex-to-triangle distance. 0 = disabled. */
+	UPROPERTY(EditAnywhere, Category = "Common Fracture|Split Islands", meta = (EditCondition = "SplitIslands", Units = "cm", ClampMin = "0.0"));
+	float VertexToSurfaceBridgeDistance = 0.f;
+
 	/**
 	 * If enabled, add extra vertices (without triangles) to the geometry in regions where vertices are spaced too far apart (e.g. across large triangles)
 	 * These extra vertices will be used as collision samples in particle-implicit collisions, and can help the physics system detect collisions more accurately
@@ -1372,10 +1412,18 @@ private:
 	UPROPERTY(EditAnywhere, Category = "Common Fracture");
 	bool GroupFracture = true;
 
-	/** Whether to split the fractured mesh pieces based on geometric connectivity after fracturing */
-	UPROPERTY(EditAnywhere, Category = "Common Fracture");
+	/** Whether to split the fractured mesh pieces into separate islands based on geometric connectivity after fracturing */
+	UPROPERTY(EditAnywhere, Category = "Common Fracture|Split Islands");
 	bool SplitIslands = true;
 
+	/** Distance threshold for connecting spatially-close vertices before computing islands */
+	UPROPERTY(EditAnywhere, Category = "Common Fracture|Split Islands", meta = (EditCondition = "SplitIslands", Units = "cm", ClampMin = "0.0"));
+	float CloseVertexDistance = 0.001f;
+
+	/** If > 0, bridge separate islands whose surfaces are within this vertex-to-triangle distance. 0 = disabled. */
+	UPROPERTY(EditAnywhere, Category = "Common Fracture|Split Islands", meta = (EditCondition = "SplitIslands", Units = "cm", ClampMin = "0.0"));
+	float VertexToSurfaceBridgeDistance = 0.f;
+
 	/** Amount of space to leave between cut pieces */
 	UPROPERTY(EditAnywhere, Category = "Common Fracture", meta = (DataflowInput, UIMin = 0.f));
 	float Grout = 0.f;
diff --git a/Engine/Plugins/Experimental/GeometryCollectionPlugin/Source/GeometryCollectionNodes/Public/Dataflow/GeometryCollectionMeshNodes.h b/Engine/Plugins/Experimental/GeometryCollectionPlugin/Source/GeometryCollectionNodes/Public/Dataflow/GeometryCollectionMeshNodes.h
--- a/Engine/Plugins/Experimental/GeometryCollectionPlugin/Source/GeometryCollectionNodes/Public/Dataflow/GeometryCollectionMeshNodes.h
+++ b/Engine/Plugins/Experimental/GeometryCollectionPlugin/Source/GeometryCollectionNodes/Public/Dataflow/GeometryCollectionMeshNodes.h
@@ -371,17 +371,21 @@ public:
 	FManagedArrayCollection Collection;
 
 	/** Whether to split the mesh into multiple bones based on the mesh connectivity */
-	UPROPERTY(EditAnywhere, Category = "General", meta = (DataflowInput));
+	UPROPERTY(EditAnywhere, Category = "General|SplitIslands", meta = (DataflowInput));
 	bool bSplitIslands = false;
 
 	/** Whether to consider coincident vertices as connected even if the topology does not connect them */
-	UPROPERTY(EditAnywhere, Category = "General");
+	UPROPERTY(EditAnywhere, Category = "General|SplitIslands", meta = (EditCondition = "bSplitIslands"));
 	bool bConnectIslandsByVertexOverlap = false;
 
 	/** Vertices closer than this distance are considered to be overlapping */
-	UPROPERTY(EditAnywhere, Category = "General", meta = (ClampMin = 0.0));
+	UPROPERTY(EditAnywhere, Category = "General|SplitIslands", meta = (EditCondition = "bSplitIslands", Units = "cm", ClampMin = 0.0));
 	float ConnectVerticesThreshold = 0.001;
 
+	/** If > 0, bridge separate islands whose surfaces are within this vertex-to-triangle distance. 0 = disabled. */
+	UPROPERTY(EditAnywhere, Category = "General|SplitIslands", meta = (EditCondition = "bSplitIslands", Units = "cm", ClampMin = 0.0));
+	float VertexToSurfaceBridgeDistance = 0.f;
+
 	/** Whether to add a root cluster for the single mesh case. Note if the mesh is split, the root cluster will always be added. */
 	UPROPERTY(EditAnywhere, Category = "General", meta = (DataflowInput));
 	bool bAddClusterRootForSingleMesh = false;
